AnalyticsET now allows AppVersion to be specified on the command line using the -AnalyticsAppVersion=XXX option. This allows playtests and UX tests to tag all runs in their event more easily so analytics can find them.

In addition, if the AppVersion string is user-specified (either by INI config or command line), it will have {VERSION} replaced by GEngineVersion to allow versioning to be present as well, even when custom AppVersions are used.

diff --git a/Engine/Source/Runtime/Analytics/AnalyticsET/Private/AnalyticsET.cpp b/Engine/Source/Runtime/Analytics/AnalyticsET/Private/AnalyticsET.cpp
--- a/Engine/Source/Runtime/Analytics/AnalyticsET/Private/AnalyticsET.cpp
+++ b/Engine/Source/Runtime/Analytics/AnalyticsET/Private/AnalyticsET.cpp
@@ -153,9 +153,12 @@ FAnalyticsProviderET::FAnalyticsProviderET(const FAnalyticsET::Config& ConfigVal
 		: ConfigValues.APIServerET;
 
 	// default to GEngineVersion if one is not provided, append GEngineVersion otherwise.
-	AppVersion = ConfigValues.AppVersionET.IsEmpty()
-		? FString::Printf(TEXT("%s"), *GEngineVersion.ToString())
-		: FString::Printf(TEXT("%s"), *ConfigValues.AppVersionET);
+	FString ConfigAppVersion = ConfigValues.AppVersionET;
+	// Allow the cmdline to force a specific AppVersion so it can be set dynamically.
+	FParse::Value(FCommandLine::Get(), TEXT("ANALYTICSAPPVERSION="), ConfigAppVersion, false);
+	AppVersion = ConfigAppVersion.IsEmpty() 
+		? GEngineVersion.ToString() 
+		: ConfigAppVersion.Replace(TEXT("{VERSION}"), *GEngineVersion.ToString(), ESearchCase::CaseSensitive);
 
 	UE_LOG(LogAnalytics, Log, TEXT("ET APIKey = %s. APIServer = %s. AppVersion = %s"), *APIKey, *APIServer, *AppVersion);
 
diff --git a/Engine/Source/Runtime/Analytics/AnalyticsET/Public/AnalyticsET.h b/Engine/Source/Runtime/Analytics/AnalyticsET/Public/AnalyticsET.h
--- a/Engine/Source/Runtime/Analytics/AnalyticsET/Public/AnalyticsET.h
+++ b/Engine/Source/Runtime/Analytics/AnalyticsET/Public/AnalyticsET.h
@@ -38,7 +38,11 @@ public:
 		FString APIKeyET;
 		/** Swrve API Server - Defaults if empty to GetDefaultAPIServer. */
 		FString APIServerET;
-		/** AppVersion - defines the app version passed to the provider. By default this will be GEngineVersion. If you provide your own, ".<GEngineVersion>" is appended to it. */
+		/** 
+		 * AppVersion - defines the app version passed to the provider. By default this will be GEngineVersion, but you can supply your own. 
+		 * As a convenience, you can use -AnalyticsAppVersion=XXX to force the AppVersion to a specific value. Useful for playtest etc where you want to define a specific version string dynamically.
+		 * If you supply your own Version string, occurrences of "{VERSION}" are replaced with GEngineVersion. ie, -AnalyticsAppVersion=MyCustomID-{VERSION}.
+		 */
 		FString AppVersionET;
 
 		/** KeyName required for APIKey configuration. */
