Mark analytic events with an IsEditor flag.

diff --git a/Engine/Source/Runtime/Analytics/AnalyticsET/Private/AnalyticsET.cpp b/Engine/Source/Runtime/Analytics/AnalyticsET/Private/AnalyticsET.cpp
--- a/Engine/Source/Runtime/Analytics/AnalyticsET/Private/AnalyticsET.cpp
+++ b/Engine/Source/Runtime/Analytics/AnalyticsET/Private/AnalyticsET.cpp
@@ -261,6 +261,7 @@ void FAnalyticsProviderET::FlushEvents()
 		JsonWriter->WriteValue(TEXT("EventName"), Entry.EventName);
 		FString DateOffset = (CurrentTime - Entry.TimeStamp).ToString();
 		JsonWriter->WriteValue(TEXT("DateOffset"), DateOffset);
+		JsonWriter->WriteValue(TEXT("IsEditor"), FString::FromInt(GIsEditor));
 		if (Entry.Attributes.Num() > 0)
 		{
 			// optional attributes for this event
@@ -282,14 +283,15 @@ void FAnalyticsProviderET::FlushEvents()
  	TSharedRef<IHttpRequest> HttpRequest = FHttpModule::Get().CreateRequest();
  	HttpRequest->SetHeader(TEXT("Content-Type"), TEXT("application/json; charset=utf-8"));
  	HttpRequest->SetURL(
-		FString::Printf(TEXT("%sCollectData.1?SessionID=%s&AppID=%s&AppVersion=%s&UserID=%s"),
+		FString::Printf(TEXT("%sCollectData.1?SessionID=%s&AppID=%s&AppVersion=%s&UserID=%s&IsEditor=%s"),
 		*APIServer, 
 		*FGenericPlatformHttp::UrlEncode(SessionID),
 		*FGenericPlatformHttp::UrlEncode(APIKey), 
 		*FGenericPlatformHttp::UrlEncode(AppVersion),
-		*FGenericPlatformHttp::UrlEncode(UserID)
+		*FGenericPlatformHttp::UrlEncode(UserID),
+		*FGenericPlatformHttp::UrlEncode(FString::FromInt(GIsEditor))
 		));
- 	HttpRequest->SetVerb(TEXT("POST"));
+   	HttpRequest->SetVerb(TEXT("POST"));
  	HttpRequest->SetContentAsString(Payload);
  	HttpRequest->OnProcessRequestComplete().BindRaw(this, &FAnalyticsProviderET::EventRequestComplete);
  	HttpRequest->ProcessRequest();
