Removed User Attribute from Analytics

#ttp 337726 - EDITOR:ANALYTICS: Remove User Attribute

#branch UE4

diff --git a/Engine/Source/Editor/UnrealEd/Private/UnrealEdMisc.cpp b/Engine/Source/Editor/UnrealEd/Private/UnrealEdMisc.cpp
--- a/Engine/Source/Editor/UnrealEd/Private/UnrealEdMisc.cpp
+++ b/Engine/Source/Editor/UnrealEd/Private/UnrealEdMisc.cpp
@@ -1280,7 +1280,6 @@ void FUnrealEdMisc::TickPerformanceSurvey()
 			FString AveFrameRateString = FString::Printf( TEXT( "%.1f" ), AveFrameRate);
 			IAnalyticsProvider& EngineAnalytics = FEngineAnalytics::GetProvider();
 			EngineAnalytics.RecordEvent(TEXT( "Editor.Performance.FrameRate" ), TEXT( "MeanFrameRate" ), AveFrameRateString);
-			EngineAnalytics.RecordUserAttribute(TEXT( "MeanFrameRate" ), AveFrameRateString);
 		}
 
 		CancelPerformanceSurvey();
diff --git a/Engine/Source/Runtime/Analytics/Analytics/Public/Interfaces/IAnalyticsProvider.h b/Engine/Source/Runtime/Analytics/Analytics/Public/Interfaces/IAnalyticsProvider.h
--- a/Engine/Source/Runtime/Analytics/Analytics/Public/Interfaces/IAnalyticsProvider.h
+++ b/Engine/Source/Runtime/Analytics/Analytics/Public/Interfaces/IAnalyticsProvider.h
@@ -148,54 +148,6 @@ public:
 		RecordEvent(EventName, Attributes);
 	}
 
-	/**
-	 * Update an array of user attributes.
-	 * 
-	 * Note that not all providers support user attributes. In this case this method
-	 * is equivalent to sending a regular event named "User Attribute".
-	 * 
-	 * @param AttributeArray - the array of attribute name/values to set.
-	 */
-	virtual void RecordUserAttribute(const TArray<FAnalyticsEventAttribute>& Attributes)
-	{
-		RecordEvent(TEXT("User Attribute"), Attributes);
-	}
-
-	/**
-	 * Overload for RecordUserAttribute that takes a single attribute name/value pair.
-	 * 
-	 * @param ParamName attribute name
-	 * @param ParamValue attribute value
-	 * 
-	 * Note that not all providers support user attributes. In this case this method
-	 * is equivalent to sending a regular event named "User Attribute".
-	 */
-	void RecordUserAttribute(const FString& ParamName, const FString& ParamValue)
-	{
-		TArray<FAnalyticsEventAttribute> Attributes;
-		Attributes.Add(FAnalyticsEventAttribute(ParamName, ParamValue));
-		RecordUserAttribute(Attributes);
-	}
-
-	/**
-	 * Overload for RecordUserAttribute that takes two attribute name/value pairs.
-	 * 
-	 * @param Param1Name attribute name
-	 * @param Param1Value attribute value
-	 * @param Param2Name attribute name
-	 * @param Param2Value attribute value
-	 * 
-	 * Note that not all providers support user attributes. In this case this method
-	 * is equivalent to sending a regular event named "User Attribute".
-	 */
-	void RecordUserAttribute(const FString& Param1Name, const FString& Param1Value, const FString& Param2Name, const FString& Param2Value)
-	{
-		TArray<FAnalyticsEventAttribute> Attributes;
-		Attributes.Add(FAnalyticsEventAttribute(Param1Name, Param1Value));
-		Attributes.Add(FAnalyticsEventAttribute(Param2Name, Param2Value));
-		RecordUserAttribute(Attributes);
-	}
-
 	/**
 	 * Record an in-game purchase of a an item.
 	 * 
diff --git a/Engine/Source/Runtime/Analytics/AnalyticsMulticast/Private/AnalyticsMulticast.cpp b/Engine/Source/Runtime/Analytics/AnalyticsMulticast/Private/AnalyticsMulticast.cpp
--- a/Engine/Source/Runtime/Analytics/AnalyticsMulticast/Private/AnalyticsMulticast.cpp
+++ b/Engine/Source/Runtime/Analytics/AnalyticsMulticast/Private/AnalyticsMulticast.cpp
@@ -29,7 +29,6 @@ public:
 	virtual bool SetSessionID(const FString& InSessionID) override;
 
 	virtual void RecordEvent(const FString& EventName, const TArray<FAnalyticsEventAttribute>& Attributes) override;
-	virtual void RecordUserAttribute(const TArray<FAnalyticsEventAttribute>& Attributes) override;
 	virtual void RecordItemPurchase(const FString& ItemId, const FString& Currency, int PerItemCost, int ItemQuantity) override;
 	virtual void RecordCurrencyPurchase(const FString& GameCurrencyType, int GameCurrencyAmount, const FString& RealCurrencyType, float RealMoneyCost, const FString& PaymentProvider) override;
 	virtual void RecordCurrencyGiven(const FString& GameCurrencyType, int GameCurrencyAmount) override;
@@ -201,14 +200,6 @@ void FAnalyticsProviderMulticast::RecordEvent( const FString& EventName, const T
 	}
 }
 
-void FAnalyticsProviderMulticast::RecordUserAttribute( const TArray<FAnalyticsEventAttribute>& Attributes )
-{
-	for (TArray<TSharedPtr<IAnalyticsProvider> >::TConstIterator it(Providers);it;++it)
-	{
-		(*it)->RecordUserAttribute(Attributes);
-	}
-}
-
 void FAnalyticsProviderMulticast::RecordItemPurchase( const FString& ItemId, const FString& Currency, int PerItemCost, int ItemQuantity )
 {
 	for (TArray<TSharedPtr<IAnalyticsProvider> >::TConstIterator it(Providers);it;++it)
diff --git a/Engine/Source/Runtime/Analytics/AnalyticsSwrve/Private/AnalyticsSwrve.cpp b/Engine/Source/Runtime/Analytics/AnalyticsSwrve/Private/AnalyticsSwrve.cpp
--- a/Engine/Source/Runtime/Analytics/AnalyticsSwrve/Private/AnalyticsSwrve.cpp
+++ b/Engine/Source/Runtime/Analytics/AnalyticsSwrve/Private/AnalyticsSwrve.cpp
@@ -33,8 +33,6 @@ public:
 	virtual bool SetSessionID(const FString& InSessionID) override;
 
 	virtual void RecordEvent(const FString& EventName, const TArray<FAnalyticsEventAttribute>& Attributes) override;
-	virtual void RecordUserAttribute(const TArray<FAnalyticsEventAttribute>& Attributes) override;
-
 	virtual void RecordItemPurchase(const FString& ItemId, const FString& Currency, int PerItemCost, int ItemQuantity) override;
 	virtual void RecordCurrencyPurchase(const FString& GameCurrencyType, int GameCurrencyAmount, const FString& RealCurrencyType, float RealMoneyCost, const FString& PaymentProvider) override;
 	virtual void RecordCurrencyGiven(const FString& GameCurrencyType, int GameCurrencyAmount) override;
@@ -363,24 +361,6 @@ void FAnalyticsProviderSwrve::RecordEvent(const FString& EventName, const TArray
 	SendToSwrve(TEXT("1/event"), FString::Printf(TEXT("name=%s"), *EventName), EventParams);
 }
 
-void FAnalyticsProviderSwrve::RecordUserAttribute(const TArray<FAnalyticsEventAttribute>& Attributes)
-{
-	if (Attributes.Num() == 0)
-	{
-		UE_LOG(LogAnalytics, Warning, TEXT("SwrveLogUserAttributeUpdateArray called with no attributes to update."));
-		return;
-	}
-
-	FString EventParams = TEXT("");
-	for (int Ndx=0;Ndx<Attributes.Num();++Ndx)
-	{
-		EventParams += FString(TEXT("&")) + Attributes[Ndx].AttrName + FString(TEXT("=")) + Attributes[Ndx].AttrValue;
-	}
-
-	SendToSwrve(TEXT("1/user"), EventParams, FString());
-}
-
-
 void FAnalyticsProviderSwrve::RecordItemPurchase( const FString& ItemId, const FString& Currency, int PerItemCost, int ItemQuantity )
 {
 	SendToSwrve(TEXT("1/purchase"), FString::Printf(TEXT("item=%s&cost=%d&quantity=%d&currency=%s"), 
