Added RecordEvent that takes a single FAnalyticsEventAttribute

Removed StartSession/RecordEvent that took multiple params (should use the array overloads instead), fixedup usage

diff --git a/Engine/Source/Editor/MainFrame/Private/MainFrameModule.cpp b/Engine/Source/Editor/MainFrame/Private/MainFrameModule.cpp
--- a/Engine/Source/Editor/MainFrame/Private/MainFrameModule.cpp
+++ b/Engine/Source/Editor/MainFrame/Private/MainFrameModule.cpp
@@ -719,10 +719,10 @@ void FMainFrameModule::HandleLevelEditorModuleCompileFinished(const FString& Log
 
 		if( FEngineAnalytics::IsAvailable() )
 		{
-			FEngineAnalytics::GetProvider().RecordEvent( 
-				TEXT("Editor.Modules.Recompile"),
-				TEXT("Duration"), FString::Printf(TEXT("%.3f"), ModuleCompileDuration),
-				TEXT("Result"), CompilationResult == ECompilationResult::Succeeded ? TEXT("Succeeded") : TEXT("Failed"));
+			TArray< FAnalyticsEventAttribute > CompileAttribs;
+			CompileAttribs.Add(FAnalyticsEventAttribute(TEXT("Duration"), FString::Printf(TEXT("%.3f"), ModuleCompileDuration)));
+			CompileAttribs.Add(FAnalyticsEventAttribute(TEXT("Result"), CompilationResult == ECompilationResult::Succeeded ? TEXT("Succeeded") : TEXT("Failed")));
+			FEngineAnalytics::GetProvider().RecordEvent(TEXT("Editor.Modules.Recompile"), CompileAttribs);
 		}
 	}
 
diff --git a/Engine/Source/Editor/UnrealEd/Private/UnrealEdMisc.cpp b/Engine/Source/Editor/UnrealEd/Private/UnrealEdMisc.cpp
--- a/Engine/Source/Editor/UnrealEd/Private/UnrealEdMisc.cpp
+++ b/Engine/Source/Editor/UnrealEd/Private/UnrealEdMisc.cpp
@@ -1353,10 +1353,10 @@ void FUnrealEdMisc::OnUserDefinedGestureChanged(const FUICommandInfo& CommandInf
 		FString GestureName = FString::Printf(TEXT("%s.%s"), *CommandInfo.GetBindingContext().ToString(), *CommandInfo.GetCommandName().ToString());
 
 		//@todo This shouldn't be using a localized value; GetInputText() [10/11/2013 justin.sargent]
-		FEngineAnalytics::GetProvider().RecordEvent(
-			TEXT("Editor.Usage.KeyboardShortcut"),
-			TEXT("Context"), GestureName,
-			TEXT("Shortcut"), CommandInfo.GetActiveGesture()->GetInputText().ToString());
+		TArray< FAnalyticsEventAttribute > GestureAttribs;
+		GestureAttribs.Add(FAnalyticsEventAttribute(TEXT("Context"), GestureName));
+		GestureAttribs.Add(FAnalyticsEventAttribute(TEXT("Shortcut"), CommandInfo.GetActiveGesture()->GetInputText().ToString()));
+		FEngineAnalytics::GetProvider().RecordEvent(TEXT("Editor.Usage.KeyboardShortcut"), GestureAttribs);
 	}
 }
 
diff --git a/Engine/Source/Runtime/Analytics/Analytics/Public/Interfaces/IAnalyticsProvider.h b/Engine/Source/Runtime/Analytics/Analytics/Public/Interfaces/IAnalyticsProvider.h
--- a/Engine/Source/Runtime/Analytics/Analytics/Public/Interfaces/IAnalyticsProvider.h
+++ b/Engine/Source/Runtime/Analytics/Analytics/Public/Interfaces/IAnalyticsProvider.h
@@ -31,31 +31,27 @@ public:
 	}
 
 	/**
-	 * Overload for StartSession that takes a single name/value pair
+	 * Overload for StartSession that takes a single attribute
 	 *
-	 * @param ParamName attribute name
-	 * @param ParamValue attribute value
+	 * @param Attribute attribute name and value
 	 */
-	bool StartSession(const FString& ParamName, const FString& ParamValue)
+	bool StartSession(const FAnalyticsEventAttribute& Attribute)
 	{
 		TArray<FAnalyticsEventAttribute> Attributes;
-		Attributes.Add(FAnalyticsEventAttribute(ParamName, ParamValue));
+		Attributes.Add(Attribute);
 		return StartSession(Attributes);
 	}
 
 	/**
-	 * Overload for StartSession that takes a two name/value pairs
+	 * Overload for StartSession that takes a single name/value pair
 	 *
-	 * @param Param1Name attribute name
-	 * @param Param1Value attribute value
-	 * @param Param2Name attribute name
-	 * @param Param2Value attribute value
+	 * @param ParamName attribute name
+	 * @param ParamValue attribute value
 	 */
-	bool StartSession(const FString& Param1Name, const FString& Param1Value, const FString& Param2Name, const FString& Param2Value)
+	bool StartSession(const FString& ParamName, const FString& ParamValue)
 	{
 		TArray<FAnalyticsEventAttribute> Attributes;
-		Attributes.Add(FAnalyticsEventAttribute(Param1Name, Param1Value));
-		Attributes.Add(FAnalyticsEventAttribute(Param2Name, Param2Value));
+		Attributes.Add(FAnalyticsEventAttribute(ParamName, ParamValue));
 		return StartSession(Attributes);
 	}
 
@@ -118,33 +114,29 @@ public:
 	}
 
 	/**
-	 * Overload for RecordEvent that takes a single name/value pair
+	 * Overload for RecordEvent that takes a single attribute
 	 *
 	 * @param EventName name of the event
-	 * @param ParamName attribute name
-	 * @param ParamValue attribute value
+	 * @param Attribute attribute name and value
 	 */
-	void RecordEvent(const FString& EventName, const FString& ParamName, const FString& ParamValue)
+	void RecordEvent(const FString& EventName, const FAnalyticsEventAttribute& Attribute)
 	{
 		TArray<FAnalyticsEventAttribute> Attributes;
-		Attributes.Add(FAnalyticsEventAttribute(ParamName, ParamValue));
+		Attributes.Add(Attribute);
 		RecordEvent(EventName, Attributes);
 	}
 
 	/**
-	 * Overload for RecordEvent that takes two name/value pairs
+	 * Overload for RecordEvent that takes a single name/value pair
 	 *
 	 * @param EventName name of the event
-	 * @param Param1Name attribute name
-	 * @param Param1Value attribute value
-	 * @param Param2Name attribute name
-	 * @param Param2Value attribute value
+	 * @param ParamName attribute name
+	 * @param ParamValue attribute value
 	 */
-	void RecordEvent(const FString& EventName, const FString& Param1Name, const FString& Param1Value, const FString& Param2Name, const FString& Param2Value)
+	void RecordEvent(const FString& EventName, const FString& ParamName, const FString& ParamValue)
 	{
 		TArray<FAnalyticsEventAttribute> Attributes;
-		Attributes.Add(FAnalyticsEventAttribute(Param1Name, Param1Value));
-		Attributes.Add(FAnalyticsEventAttribute(Param2Name, Param2Value));
+		Attributes.Add(FAnalyticsEventAttribute(ParamName, ParamValue));
 		RecordEvent(EventName, Attributes);
 	}
 
