Changes necessary to use the new DataRouter backend while simultaneously using the existing data collector.

* Based off of changed originally made by [REDACTED_PERSON]

* This first pass turns on the new DataRouter path when compiled internally only for now

diff --git a/Engine/Source/Runtime/Analytics/AnalyticsET/Private/AnalyticsET.cpp b/Engine/Source/Runtime/Analytics/AnalyticsET/Private/AnalyticsET.cpp
--- a/Engine/Source/Runtime/Analytics/AnalyticsET/Private/AnalyticsET.cpp
+++ b/Engine/Source/Runtime/Analytics/AnalyticsET/Private/AnalyticsET.cpp
@@ -59,7 +59,10 @@ private:
 	FString BuildType;
 	/** The AppVersion passed to ET. */
 	FString AppVersion;
-
+	/** True if we are sending to the data router*/
+	bool UseDataRouter;
+	/** The URL to which uploads are sent when using the data router*/
+	FString DataRouterUploadURL;
 	/** Max number of analytics events to cache before pushing to server */
 	const int32 MaxCachedNumEvents;
 	/** Max time that can elapse before pushing cached events to server */
@@ -89,10 +92,18 @@ private:
 	/** List of analytic events pending a server update */
 	TArray<FAnalyticsEventEntry> CachedEvents;
 
+	/** Critical section for updating the CachedEvents */
+	FCriticalSection CachedEventsCS;
+
 	/**
 	 * Delegate called when an event Http request completes
 	 */
 	void EventRequestComplete(FHttpRequestPtr HttpRequest, FHttpResponsePtr HttpResponse, bool bSucceeded);
+
+	/**
+	 * Delegate called when an event Http request completes (for DataRouter)
+	 */
+	void EventRequestCompleteDataRouter(FHttpRequestPtr HttpRequest, FHttpResponsePtr HttpResponse, bool bSucceeded);
 };
 
 void FAnalyticsET::StartupModule()
@@ -113,6 +124,8 @@ TSharedPtr<IAnalyticsProvider> FAnalyticsET::CreateAnalyticsProvider(const FAnal
 		ConfigValues.APIKeyET = GetConfigValue.Execute(Config::GetKeyNameForAPIKey(), true);
 		ConfigValues.APIServerET = GetConfigValue.Execute(Config::GetKeyNameForAPIServer(), false);
 		ConfigValues.AppVersionET = GetConfigValue.Execute(Config::GetKeyNameForAppVersion(), false);
+		ConfigValues.UseDataRouterET = GetConfigValue.Execute(Config::GetKeyNameForUseDataRouter(), false);
+		ConfigValues.DataRouterUploadURLET = GetConfigValue.Execute(Config::GetKeyNameForDataRouterUploadURL(), false);
 		return CreateAnalyticsProvider(ConfigValues);
 	}
 	else
@@ -154,6 +167,14 @@ FAnalyticsProviderET::FAnalyticsProviderET(const FAnalyticsET::Config& ConfigVal
 		? FAnalyticsET::Config::GetDefaultAPIServer()
 		: ConfigValues.APIServerET;
 
+	// allow the DataRouterUploadURL value to be empty and use defaults.
+	DataRouterUploadURL = ConfigValues.DataRouterUploadURLET.IsEmpty()
+		? FAnalyticsET::Config::GetDefaultDataRouterUploadURL()
+		: ConfigValues.DataRouterUploadURLET;
+
+	// determine if we are using the data router
+	UseDataRouter = FCString::ToBool(*ConfigValues.UseDataRouterET);
+
 	// default to GEngineVersion if one is not provided, append GEngineVersion otherwise.
 	FString ConfigAppVersion = ConfigValues.AppVersionET;
 	// Allow the cmdline to force a specific AppVersion so it can be set dynamically.
@@ -162,7 +183,15 @@ FAnalyticsProviderET::FAnalyticsProviderET(const FAnalyticsET::Config& ConfigVal
 		? GEngineVersion.ToString() 
 		: ConfigAppVersion.Replace(TEXT("%VERSION%"), *GEngineVersion.ToString(), ESearchCase::CaseSensitive);
 
-	UE_LOG(LogAnalytics, Log, TEXT("ET APIKey = %s. APIServer = %s. AppVersion = %s"), *APIKey, *APIServer, *AppVersion);
+	if (UseDataRouter)
+	{
+		UE_LOG(LogAnalytics, Log, TEXT("ET APIKey = %s. DataRouterUploadURL = %s. AppVersion = %s"), *APIKey, *DataRouterUploadURL, *AppVersion);
+	}
+	else
+	{		
+		UE_LOG(LogAnalytics, Log, TEXT("ET APIKey = %s. APIServer = %s. AppVersion = %s"), *APIKey, *APIServer, *AppVersion);
+	}
+	
 
 	// cache the build type string
 	FAnalytics::BuildType BuildTypeEnum = FAnalytics::Get().GetBuildType();
@@ -248,6 +277,10 @@ void FAnalyticsProviderET::EndSession()
 
 void FAnalyticsProviderET::FlushEvents()
 {
+	// There are much better ways to do this, but since most events are recorded and handled on the same (game) thread,
+	// this is probably mostly fine for now, and simply favoring not crashing at the moment
+	FScopeLock ScopedLock(&CachedEventsCS);
+
 	if(ensure(FModuleManager::Get().IsModuleLoaded("HTTP")))
 	{
 		FString Payload;
@@ -256,6 +289,11 @@ void FAnalyticsProviderET::FlushEvents()
 
 		TSharedRef< TJsonWriter<TCHAR, TCondensedJsonPrintPolicy<TCHAR> > > JsonWriter = TJsonWriterFactory<TCHAR, TCondensedJsonPrintPolicy<TCHAR> >::Create(&Payload);
 		JsonWriter->WriteObjectStart();
+		if (UseDataRouter)
+		{
+			JsonWriter->WriteValue(TEXT("SessionID"), SessionID);
+			JsonWriter->WriteValue(TEXT("UserID"), UserID);
+		}
 		JsonWriter->WriteArrayStart(TEXT("Events"));
 		for (int32 EventIdx = 0; EventIdx < CachedEvents.Num(); EventIdx++)
 		{
@@ -284,9 +322,10 @@ void FAnalyticsProviderET::FlushEvents()
 		UE_LOG(LogAnalytics, Verbose, TEXT("ET Flush: Payload:\n%s"), *Payload);
 
 		// Create/send Http request for an event
- 		TSharedRef<IHttpRequest> HttpRequest = FHttpModule::Get().CreateRequest();
- 		HttpRequest->SetHeader(TEXT("Content-Type"), TEXT("application/json; charset=utf-8"));
- 		HttpRequest->SetURL(
+		TSharedRef<IHttpRequest> HttpRequest = FHttpModule::Get().CreateRequest();
+		HttpRequest->SetHeader(TEXT("Content-Type"), TEXT("application/json; charset=utf-8"));
+
+		HttpRequest->SetURL(
 			FString::Printf(TEXT("%sCollectData.1?SessionID=%s&AppID=%s&AppVersion=%s&UserID=%s&IsEditor=%s"),
 			*APIServer, 
 			*FGenericPlatformHttp::UrlEncode(SessionID),
@@ -294,12 +333,41 @@ void FAnalyticsProviderET::FlushEvents()
 			*FGenericPlatformHttp::UrlEncode(AppVersion),
 			*FGenericPlatformHttp::UrlEncode(UserID),
 			*FGenericPlatformHttp::UrlEncode(FString::FromInt(GIsEditor))
-			));
-   		HttpRequest->SetVerb(TEXT("POST"));
- 		HttpRequest->SetContentAsString(Payload);
- 		HttpRequest->OnProcessRequestComplete().BindRaw(this, &FAnalyticsProviderET::EventRequestComplete);
+		));
+
+		HttpRequest->SetVerb(TEXT("POST"));
+		HttpRequest->SetContentAsString(Payload);
+		HttpRequest->OnProcessRequestComplete().BindRaw(this, &FAnalyticsProviderET::EventRequestComplete);
+
  		HttpRequest->ProcessRequest();
 
+		if (UseDataRouter)
+		{
+			// If we're using the DataRouter backend, then submit the same request to the new DataRouter backend
+			// NOTE - This branch is temp, and we will eventually use the DataRouter path exclusively
+
+			// Create/send Http request for an event
+			TSharedRef<IHttpRequest> HttpRequest = FHttpModule::Get().CreateRequest();
+			HttpRequest->SetHeader(TEXT("Content-Type"), TEXT("application/json; charset=utf-8"));
+
+			// TODO need agent here??
+			HttpRequest->SetURL(
+				FString::Printf(TEXT("%s?SessionID=%s&AppID=%s&AppVersion=%s&UserID=%s&IsEditor=%s&AppEnvironment=%s&UploadType=eteventstream"),
+				*DataRouterUploadURL,
+				*FGenericPlatformHttp::UrlEncode(SessionID),
+				*FGenericPlatformHttp::UrlEncode(APIKey), 
+				*FGenericPlatformHttp::UrlEncode(AppVersion),
+				*FGenericPlatformHttp::UrlEncode(UserID),
+				*FGenericPlatformHttp::UrlEncode(FString::FromInt(GIsEditor)),
+				*FGenericPlatformHttp::UrlEncode(BuildType)
+			));
+
+			HttpRequest->SetVerb(TEXT("POST"));
+			HttpRequest->SetContentAsString(Payload);
+			HttpRequest->OnProcessRequestComplete().BindRaw(this, &FAnalyticsProviderET::EventRequestCompleteDataRouter);
+			HttpRequest->ProcessRequest();
+		}
+
 		FlushEventsCountdown = MaxCachedElapsedTime;
 		CachedEvents.Empty();
 	}
@@ -348,6 +416,9 @@ void FAnalyticsProviderET::RecordEvent(const FString& EventName, const TArray<FA
 		UE_LOG(LogAnalytics, Log, TEXT("Event %s has too many attributes (%d). May be truncated at the collector."), *EventName, Attributes.Num());
 	}
 
+	// There are much better ways to do this, but since most events are recorded and handled on the same (game) thread,
+	// this is probably mostly fine for now, and simply favoring not crashing at the moment
+	FScopeLock ScopedLock(&CachedEventsCS);
 	CachedEvents.Add(FAnalyticsEventEntry(EventName, Attributes));
 }
 
@@ -357,11 +428,21 @@ void FAnalyticsProviderET::EventRequestComplete(FHttpRequestPtr HttpRequest, FHt
 	if (bSucceeded && HttpResponse.IsValid())
 	{
 		UE_LOG(LogAnalytics, VeryVerbose, TEXT("ET response for [%s]. Code: %d. Payload: %s"), *HttpRequest->GetURL(), HttpResponse->GetResponseCode(), *HttpResponse->GetContentAsString());
-		FString PayloadStr = HttpResponse->GetContentAsString();
-		bool test = false;
 	}
 	else
 	{
 		UE_LOG(LogAnalytics, VeryVerbose, TEXT("ET response for [%s]. No response"), *HttpRequest->GetURL());
 	}
 }
+
+void FAnalyticsProviderET::EventRequestCompleteDataRouter(FHttpRequestPtr HttpRequest, FHttpResponsePtr HttpResponse, bool bSucceeded)
+{
+	if (bSucceeded && HttpResponse.IsValid())
+	{
+		UE_LOG(LogAnalytics, VeryVerbose, TEXT("ET (DataRouter) response for [%s]. Code: %d. Payload: %s"), *HttpRequest->GetURL(), HttpResponse->GetResponseCode(), *HttpResponse->GetContentAsString());
+	}
+	else
+	{
+		UE_LOG(LogAnalytics, VeryVerbose, TEXT("ET (DataRouter) response for [%s]. No response"), *HttpRequest->GetURL());
+	}
+}
diff --git a/Engine/Source/Runtime/Analytics/AnalyticsET/Public/AnalyticsET.h b/Engine/Source/Runtime/Analytics/AnalyticsET/Public/AnalyticsET.h
--- a/Engine/Source/Runtime/Analytics/AnalyticsET/Public/AnalyticsET.h
+++ b/Engine/Source/Runtime/Analytics/AnalyticsET/Public/AnalyticsET.h
@@ -36,7 +36,7 @@ public:
 	{
 		/** ET APIKey - Get from your account manager */
 		FString APIKeyET;
-		/** Swrve API Server - Defaults if empty to GetDefaultAPIServer. */
+		/** ET API Server - Defaults if empty to GetDefaultAPIServer. */
 		FString APIServerET;
 		/** 
 		 * AppVersion - defines the app version passed to the provider. By default this will be GEngineVersion, but you can supply your own. 
@@ -44,6 +44,10 @@ public:
 		 * If you supply your own Version string, occurrences of "%VERSION%" are replaced with GEngineVersion. ie, -AnalyticsAppVersion=MyCustomID-%VERSION%.
 		 */
 		FString AppVersionET;
+		/** UseDataRouter - flag indicating that we should use the new datarouter system. */
+		FString UseDataRouterET;
+		/** DataRouterUploadURLET - data router upload URL - Defaults if empty to GetDefaultDataRouterUploadURL. */
+		FString DataRouterUploadURLET;
 
 		/** KeyName required for APIKey configuration. */
 		static FString GetKeyNameForAPIKey() { return TEXT("APIKeyET"); }
@@ -53,6 +57,12 @@ public:
 		static FString GetKeyNameForAppVersion() { return TEXT("AppVersionET"); }
 		/** Default value if no APIServer configuration is provided. */
 		static FString GetDefaultAPIServer() { return TEXT("http://devonline-02:/ETAP/"); }
+		/** Default value if no DataRouterServer configuration is provided. */
+		static FString GetDefaultDataRouterUploadURL() { return TEXT("https://datarouter-public-service-gamedev.ol.epicgames.net/datarouter/api/v1/public/data"); }
+		/** KeyName indicating if we should use the data router. */
+		static FString GetKeyNameForUseDataRouter() { return TEXT("UseDataRouterET"); }
+		/** KeyName for the data router upload URL. */
+		static FString GetKeyNameForDataRouterUploadURL() { return TEXT("DataRouterUploadURLET"); }
 	};
 
 	//--------------------------------------------------------------------------
