Analtyics ET updates:
* Disable caching of events using -AnalyticsDisableCaching commandline. #UE-7593
** This will also force analytics payloads to be logged.
** Together this allows testing analytics events without waiting for the caching delay
** Also allows any automation tools to look for entries with "LogAnalytics: AnalyticsET Payload:" in the log. There will only be one event per log entry this way.
* All SessionStart events will have UniqueDeviceID and Platform attributes appended to them. #UETOOL-241 #UETOOL-242.
* Fix code to always log API keys properly even when using the Data Router.
* remove legacy warning when there are more than 40 attributes in an event.

diff --git a/Engine/Source/Runtime/Analytics/AnalyticsET/Private/AnalyticsET.cpp b/Engine/Source/Runtime/Analytics/AnalyticsET/Private/AnalyticsET.cpp
--- a/Engine/Source/Runtime/Analytics/AnalyticsET/Private/AnalyticsET.cpp
+++ b/Engine/Source/Runtime/Analytics/AnalyticsET/Private/AnalyticsET.cpp
@@ -30,7 +30,6 @@ public:
 
 	virtual bool StartSession(const TArray<FAnalyticsEventAttribute>& Attributes) override;
 	virtual void EndSession() override;
-	/** ET PC implementation doesn't cache events */
 	virtual void FlushEvents() override;
 
 	virtual void SetUserID(const FString& InUserID) override;
@@ -67,6 +66,8 @@ private:
 	const int32 MaxCachedNumEvents;
 	/** Max time that can elapse before pushing cached events to server */
 	const float MaxCachedElapsedTime;
+	/** Allows events to not be cached when -AnalyticsDisableCaching is used. This should only be used for debugging as caching significantly reduces bandwidth overhead per event. */
+	bool bShouldCacheEvents;
 	/** Current countdown timer to keep track of MaxCachedElapsedTime push */
 	float FlushEventsCountdown;
 	/**
@@ -153,6 +154,7 @@ FAnalyticsProviderET::FAnalyticsProviderET(const FAnalyticsET::Config& ConfigVal
 	:bSessionInProgress(false)
 	, MaxCachedNumEvents(20)
 	, MaxCachedElapsedTime(60.0f)
+	, bShouldCacheEvents(!FParse::Param(FCommandLine::Get(), TEXT("ANALYTICSDISABLECACHING")))
 	, FlushEventsCountdown(MaxCachedElapsedTime)
 {
 	UE_LOG(LogAnalytics, Verbose, TEXT("Initializing ET Analytics provider"));
@@ -183,15 +185,11 @@ FAnalyticsProviderET::FAnalyticsProviderET(const FAnalyticsET::Config& ConfigVal
 		? GEngineVersion.ToString() 
 		: ConfigAppVersion.Replace(TEXT("%VERSION%"), *GEngineVersion.ToString(), ESearchCase::CaseSensitive);
 
+	UE_LOG(LogAnalytics, Log, TEXT("ET APIKey = %s. APIServer = %s. AppVersion = %s"), *APIKey, *APIServer, *AppVersion);
 	if (UseDataRouter)
 	{
 		UE_LOG(LogAnalytics, Log, TEXT("ET APIKey = %s. DataRouterUploadURL = %s. AppVersion = %s"), *APIKey, *DataRouterUploadURL, *AppVersion);
 	}
-	else
-	{		
-		UE_LOG(LogAnalytics, Log, TEXT("ET APIKey = %s. APIServer = %s. AppVersion = %s"), *APIKey, *APIServer, *AppVersion);
-	}
-	
 
 	// cache the build type string
 	FAnalytics::BuildType BuildTypeEnum = FAnalytics::Get().GetBuildType();
@@ -256,7 +254,14 @@ bool FAnalyticsProviderET::StartSession(const TArray<FAnalyticsEventAttribute>&
 	FPlatformMisc::CreateGuid(SessionGUID);
 	SessionID = SessionGUID.ToString(EGuidFormats::DigitsWithHyphensInBraces);
 
-	RecordEvent(TEXT("SessionStart"), Attributes);
+	// always ensure we send a few specific attributes on session start.
+	TArray<FAnalyticsEventAttribute> AppendedAttributes(Attributes);
+	// this is for legacy reasons (we used to use this ID, so helps us create old->new mappings).
+	AppendedAttributes.Emplace(TEXT("UniqueDeviceId"), FPlatformMisc::GetUniqueDeviceId());
+	// we should always know what platform is hosting this session.
+	AppendedAttributes.Emplace(TEXT("Platform"), FString(FPlatformProperties::IniPlatformName()));
+
+	RecordEvent(TEXT("SessionStart"), AppendedAttributes);
 	bSessionInProgress = !UserID.IsEmpty();
 	return bSessionInProgress;
 }
@@ -277,6 +282,12 @@ void FAnalyticsProviderET::EndSession()
 
 void FAnalyticsProviderET::FlushEvents()
 {
+	// Make sure we don't try to flush too many times. When we are not caching events it's possible this can be called when there are no events in the array.
+	if (CachedEvents.Num() == 0)
+	{
+		return;
+	}
+
 	// There are much better ways to do this, but since most events are recorded and handled on the same (game) thread,
 	// this is probably mostly fine for now, and simply favoring not crashing at the moment
 	FScopeLock ScopedLock(&CachedEventsCS);
@@ -319,7 +330,15 @@ void FAnalyticsProviderET::FlushEvents()
 		JsonWriter->WriteObjectEnd();
 		JsonWriter->Close();
 
-		UE_LOG(LogAnalytics, Verbose, TEXT("ET Flush: Payload:\n%s"), *Payload);
+		// when we are not caching events, assume we are debugging and always log analytics payloads.
+		if (bShouldCacheEvents)
+		{
+			UE_LOG(LogAnalytics, Verbose, TEXT("AnalyticsET Payload:%s"), *Payload);
+		}
+		else
+		{
+			UE_LOG(LogAnalytics, Log, TEXT("AnalyticsET Payload:%s"), *Payload);
+		}
 
 		// Create/send Http request for an event
 		TSharedRef<IHttpRequest> HttpRequest = FHttpModule::Get().CreateRequest();
@@ -411,15 +430,15 @@ bool FAnalyticsProviderET::SetSessionID(const FString& InSessionID)
 /** Helper to log any ET event. Used by all the LogXXX functions. */
 void FAnalyticsProviderET::RecordEvent(const FString& EventName, const TArray<FAnalyticsEventAttribute>& Attributes)
 {
-	if (Attributes.Num() > 40)
-	{
-		UE_LOG(LogAnalytics, Log, TEXT("Event %s has too many attributes (%d). May be truncated at the collector."), *EventName, Attributes.Num());
-	}
-
 	// There are much better ways to do this, but since most events are recorded and handled on the same (game) thread,
 	// this is probably mostly fine for now, and simply favoring not crashing at the moment
 	FScopeLock ScopedLock(&CachedEventsCS);
 	CachedEvents.Add(FAnalyticsEventEntry(EventName, Attributes));
+	// if we aren't caching events, flush immediately. This is really only for debugging as it will significantly affect bandwidth.
+	if (!bShouldCacheEvents)
+	{
+		FlushEvents();
+	}
 }
 
 
