Analytics will now REALLY log in all apps when using -AnalyticsDisableCaching.
* Log the payload at Display verbosity instead of log, as the default was Display.

diff --git a/Engine/Source/Runtime/Analytics/AnalyticsET/Private/AnalyticsET.cpp b/Engine/Source/Runtime/Analytics/AnalyticsET/Private/AnalyticsET.cpp
--- a/Engine/Source/Runtime/Analytics/AnalyticsET/Private/AnalyticsET.cpp
+++ b/Engine/Source/Runtime/Analytics/AnalyticsET/Private/AnalyticsET.cpp
@@ -337,7 +337,7 @@ void FAnalyticsProviderET::FlushEvents()
 		}
 		else
 		{
-			UE_LOG(LogAnalytics, Log, TEXT("AnalyticsET Payload:%s"), *Payload);
+			UE_LOG(LogAnalytics, Display, TEXT("AnalyticsET Payload:%s"), *Payload);
 		}
 
 		// Create/send Http request for an event
