AnalyticsEventAttribute.h
One of the template <typename T> analytic event attribute functions was missing a const otherwise it will not compile since it was taking in a const reference variable and trying to use it as just a referenced one.

diff --git a/Engine/Source/Runtime/Analytics/Analytics/Public/AnalyticsEventAttribute.h b/Engine/Source/Runtime/Analytics/Analytics/Public/AnalyticsEventAttribute.h
--- a/Engine/Source/Runtime/Analytics/Analytics/Public/AnalyticsEventAttribute.h
+++ b/Engine/Source/Runtime/Analytics/Analytics/Public/AnalyticsEventAttribute.h
@@ -52,7 +52,7 @@ struct FAnalyticsEventAttribute
 		,AttrValue(FString())
 	{
 		// Serialize the array into "value1,value2,..." format
-		for (T& Value : InValueArray)
+		for (const T& Value : InValueArray)
 		{
 			AttrValue += TTypeToString<T>::ToString(Value) + ",";
 		}
