Properly use UE_SET_LOG_VERBOSITY for LogAnalytics.

diff --git a/Engine/Source/Runtime/Analytics/AnalyticsET/Private/AnalyticsET.cpp b/Engine/Source/Runtime/Analytics/AnalyticsET/Private/AnalyticsET.cpp
--- a/Engine/Source/Runtime/Analytics/AnalyticsET/Private/AnalyticsET.cpp
+++ b/Engine/Source/Runtime/Analytics/AnalyticsET/Private/AnalyticsET.cpp
@@ -160,7 +160,7 @@ FAnalyticsProviderET::FAnalyticsProviderET(const FAnalyticsET::Config& ConfigVal
 	// if we are not caching events, we are operating in debug mode. Turn on super-verbose analytics logging
 	if (!bShouldCacheEvents)
 	{
-		LogAnalytics.SetVerbosity(ELogVerbosity::VeryVerbose);
+		UE_SET_LOG_VERBOSITY(LogAnalytics, VeryVerbose);
 	}
 
 	APIKey = ConfigValues.APIKeyET;
