Analytics Payload logging no longer shows the URL with UrlEncoded parameters, which makes reading the log easier to match up with what the collector will actually receive. This means pipes, braces, etc will show up correctly instead of with %7B and %7D etc:

Before: AnalyticsET URL:CollectData.1?SessionID=%7BFA7D7BED-4C4F-A77F-48A7-A0944CF38478%7D&AppID=...
  After: AnalyticsET URL:CollectData.1?SessionID={E5F58A26-400A-E445-936A-0DB1FA2E65B8}&AppID=...

diff --git a/Engine/Source/Runtime/Analytics/AnalyticsET/Private/AnalyticsET.cpp b/Engine/Source/Runtime/Analytics/AnalyticsET/Private/AnalyticsET.cpp
--- a/Engine/Source/Runtime/Analytics/AnalyticsET/Private/AnalyticsET.cpp
+++ b/Engine/Source/Runtime/Analytics/AnalyticsET/Private/AnalyticsET.cpp
@@ -339,7 +339,17 @@ void FAnalyticsProviderET::FlushEvents()
 			*FGenericPlatformHttp::UrlEncode(APIKey),
 			*FGenericPlatformHttp::UrlEncode(AppVersion),
 			*FGenericPlatformHttp::UrlEncode(UserID));
-		
+
+		// Recreate the URLPath for logging because we do not want to escape the parameters when logging.
+		// We cannot simply UrlEncode the entire Path after logging it because UrlEncode(Params) != UrlEncode(Param1) & UrlEncode(Param2) ...
+		UE_LOG(LogAnalytics, VeryVerbose, TEXT("[%s] AnalyticsET URL:CollectData.1?SessionID=%s&AppID=%s&AppVersion=%s&UserID=%s. Payload:%s"), 
+			*APIKey, 
+			*SessionID,
+			*APIKey,
+			*AppVersion,
+			*UserID,
+			*Payload);
+
 		// Create/send Http request for an event
 		TSharedRef<IHttpRequest> HttpRequest = FHttpModule::Get().CreateRequest();
 		HttpRequest->SetHeader(TEXT("Content-Type"), TEXT("application/json; charset=utf-8"));
@@ -349,8 +359,6 @@ void FAnalyticsProviderET::FlushEvents()
 		HttpRequest->SetContentAsString(Payload);
 		HttpRequest->OnProcessRequestComplete().BindRaw(this, &FAnalyticsProviderET::EventRequestComplete);
 
-		UE_LOG(LogAnalytics, Display, TEXT("[%s] AnalyticsET URL:%s. Payload:%s"), *APIKey, *URLPath, *Payload);
-
  		HttpRequest->ProcessRequest();
 
 		if (UseDataRouter)
