Added the new functions needed by GameAnalytics for their plugin along with feedback from Prime31

diff --git a/Engine/Plugins/Runtime/Analytics/AnalyticsBlueprintLibrary/Source/AnalyticsBlueprintLibrary/Classes/AnalyticsBlueprintLibrary.h b/Engine/Plugins/Runtime/Analytics/AnalyticsBlueprintLibrary/Source/AnalyticsBlueprintLibrary/Classes/AnalyticsBlueprintLibrary.h
--- a/Engine/Plugins/Runtime/Analytics/AnalyticsBlueprintLibrary/Source/AnalyticsBlueprintLibrary/Classes/AnalyticsBlueprintLibrary.h
+++ b/Engine/Plugins/Runtime/Analytics/AnalyticsBlueprintLibrary/Source/AnalyticsBlueprintLibrary/Classes/AnalyticsBlueprintLibrary.h
@@ -55,6 +55,22 @@ class UAnalyticsBlueprintLibrary :
 	UFUNCTION(BlueprintCallable, Category="Analytics")
 	static void RecordItemPurchase(const FString& ItemId, const FString& Currency, int32 PerItemCost, int32 ItemQuantity);
 
+	/** Records an in-game item was purchased */
+	UFUNCTION(BlueprintCallable, Category="Analytics")
+	static void RecordSimpleItemPurchase(const FString& ItemId, int32 ItemQuantity);
+
+	/** Records an in-game item was purchased with attributes */
+	UFUNCTION(BlueprintCallable, Category="Analytics")
+	static void RecordSimpleItemPurchaseWithAttributes(const FString& ItemId, int32 ItemQuantity, const TArray<FAnalyticsEventAttr>& Attributes);
+
+	/** Records an in-game currency was purchased using real-world money */
+	UFUNCTION(BlueprintCallable, Category="Analytics")
+	static void RecordSimpleCurrencyPurchase(const FString& GameCurrencyType, int32 GameCurrencyAmount);
+
+	/** Records an in-game currency was purchased using real-world money */
+	UFUNCTION(BlueprintCallable, Category="Analytics")
+	static void RecordSimpleCurrencyPurchaseWithAttributes(const FString& GameCurrencyType, int32 GameCurrencyAmount, const TArray<FAnalyticsEventAttr>& Attributes);
+
 	/** Records an in-game currency was purchased using real-world money */
 	UFUNCTION(BlueprintCallable, Category="Analytics")
 	static void RecordCurrencyPurchase(const FString& GameCurrencyType, int32 GameCurrencyAmount, const FString& RealCurrencyType, float RealMoneyCost, const FString& PaymentProvider);
@@ -63,6 +79,10 @@ class UAnalyticsBlueprintLibrary :
 	UFUNCTION(BlueprintCallable, Category="Analytics")
 	static void RecordCurrencyGiven(const FString& GameCurrencyType, int32 GameCurrencyAmount);
 
+	/** Records an in-game currency was granted by the game with no real-world money being involved */
+	UFUNCTION(BlueprintCallable, Category="Analytics")
+	static void RecordCurrencyGivenWithAttributes(const FString& GameCurrencyType, int32 GameCurrencyAmount, const TArray<FAnalyticsEventAttr>& Attributes);
+
 	/** Builds a struct from the attribute name and value */
 	UFUNCTION(BlueprintPure, Category="Analytics")
 	static FAnalyticsEventAttr MakeEventAttribute(const FString& AttributeName, const FString& AttributeValue);
@@ -82,4 +102,40 @@ class UAnalyticsBlueprintLibrary :
 	/** Sets the user id (if supported) on the analytics provider */
 	UFUNCTION(BlueprintCallable, Category="Analytics")
 	static void SetUserId(const FString& UserId);
+
+	/** Sets the user's age (if supported) on the analytics provider */
+	UFUNCTION(BlueprintCallable, Category="Analytics")
+	static void SetAge(int32 Age);
+
+	/** Sets the user's location (if supported) on the analytics provider */
+	UFUNCTION(BlueprintCallable, Category="Analytics")
+	static void SetLocation(const FString& Location);
+
+	/** Sets the user's gender (if supported) on the analytics provider */
+	UFUNCTION(BlueprintCallable, Category="Analytics")
+	static void SetGender(const FString& Gender);
+
+	/** Sets the game's build info (if supported) on the analytics provider */
+	UFUNCTION(BlueprintCallable, Category="Analytics")
+	static void SetBuildInfo(const FString& BuildInfo);
+
+	/** Records an error event has happened with attributes */
+	UFUNCTION(BlueprintCallable, Category="Analytics")
+	static void RecordErrorWithAttributes(const FString& Error, const TArray<FAnalyticsEventAttr>& Attributes);
+
+	/** Records an error event has happened */
+	UFUNCTION(BlueprintCallable, Category="Analytics")
+	static void RecordError(const FString& Error);
+
+	/** Records a user progress event has happened with a full list of progress hierarchy names and with attributes */
+	UFUNCTION(BlueprintCallable, Category="Analytics")
+	static void RecordProgressWithFullHierarchyAndAttributes(const FString& ProgressType, const TArray<FString>& ProgressNames, const TArray<FAnalyticsEventAttr>& Attributes);
+
+	/** Records a user progress event has happened with attributes */
+	UFUNCTION(BlueprintCallable, Category="Analytics")
+	static void RecordProgressWithAttributes(const FString& ProgressType, const FString& ProgressName, const TArray<FAnalyticsEventAttr>& Attributes);
+
+	/** Records a user progress event has happened */
+	UFUNCTION(BlueprintCallable, Category="Analytics")
+	static void RecordProgress(const FString& ProgressType, const FString& ProgressName);
 };
diff --git a/Engine/Plugins/Runtime/Analytics/AnalyticsBlueprintLibrary/Source/AnalyticsBlueprintLibrary/Private/AnalyticsBlueprintLibrary.cpp b/Engine/Plugins/Runtime/Analytics/AnalyticsBlueprintLibrary/Source/AnalyticsBlueprintLibrary/Private/AnalyticsBlueprintLibrary.cpp
--- a/Engine/Plugins/Runtime/Analytics/AnalyticsBlueprintLibrary/Source/AnalyticsBlueprintLibrary/Private/AnalyticsBlueprintLibrary.cpp
+++ b/Engine/Plugins/Runtime/Analytics/AnalyticsBlueprintLibrary/Source/AnalyticsBlueprintLibrary/Private/AnalyticsBlueprintLibrary.cpp
@@ -143,11 +143,71 @@ void UAnalyticsBlueprintLibrary::RecordItemPurchase(const FString& ItemId, const
 		static bool bHasLogged = false;
 		if (!bHasLogged)
 		{
-			UE_LOG(LogAnalyticsBPLib, Warning, TEXT("RecordEventWithAttributes: Failed to get the default analytics provider. Double check your [Analytics] configuration in your INI"));
+			UE_LOG(LogAnalyticsBPLib, Warning, TEXT("RecordItemPurchase: Failed to get the default analytics provider. Double check your [Analytics] configuration in your INI"));
+		}
+	}
+}
+
+void UAnalyticsBlueprintLibrary::RecordSimpleItemPurchase(const FString& ItemId, int32 ItemQuantity)
+{
+	TSharedPtr<IAnalyticsProvider> Provider = FAnalytics::Get().GetDefaultConfiguredProvider();
+	if (Provider.IsValid())
+	{
+		Provider->RecordItemPurchase(ItemId, ItemQuantity);
+	}
+	else
+	{
+		static bool bHasLogged = false;
+		if (!bHasLogged)
+		{
+			UE_LOG(LogAnalyticsBPLib, Warning, TEXT("RecordSimpleItemPurchase: Failed to get the default analytics provider. Double check your [Analytics] configuration in your INI"));
+		}
+	}
+}
+
+void UAnalyticsBlueprintLibrary::RecordSimpleItemPurchaseWithAttributes(const FString& ItemId, int32 ItemQuantity, const TArray<FAnalyticsEventAttr>& Attributes)
+{
+	TSharedPtr<IAnalyticsProvider> Provider = FAnalytics::Get().GetDefaultConfiguredProvider();
+	if (Provider.IsValid())
+	{
+		Provider->RecordItemPurchase(ItemId, ItemQuantity, ConvertAttrs(Attributes));
+	}
+	else
+	{
+		static bool bHasLogged = false;
+		if (!bHasLogged)
+		{
+			UE_LOG(LogAnalyticsBPLib, Warning, TEXT("RecordSimpleItemPurchaseWithAttributes: Failed to get the default analytics provider. Double check your [Analytics] configuration in your INI"));
 		}
 	}
 }
 
+void UAnalyticsBlueprintLibrary::RecordSimpleCurrencyPurchase(const FString& GameCurrencyType, int32 GameCurrencyAmount)
+{
+	TSharedPtr<IAnalyticsProvider> Provider = FAnalytics::Get().GetDefaultConfiguredProvider();
+	if (Provider.IsValid())
+	{
+		Provider->RecordCurrencyPurchase(GameCurrencyType, GameCurrencyAmount);
+	}
+	else
+	{
+		UE_LOG(LogAnalyticsBPLib, Warning, TEXT("RecordSimpleCurrencyPurchase: Failed to get the default analytics provider. Double check your [Analytics] configuration in your INI"));
+	}
+}
+
+void UAnalyticsBlueprintLibrary::RecordSimpleCurrencyPurchaseWithAttributes(const FString& GameCurrencyType, int32 GameCurrencyAmount, const TArray<FAnalyticsEventAttr>& Attributes)
+{
+	TSharedPtr<IAnalyticsProvider> Provider = FAnalytics::Get().GetDefaultConfiguredProvider();
+	if (Provider.IsValid())
+	{
+		Provider->RecordCurrencyPurchase(GameCurrencyType, GameCurrencyAmount, ConvertAttrs(Attributes));
+	}
+	else
+	{
+		UE_LOG(LogAnalyticsBPLib, Warning, TEXT("RecordSimpleCurrencyPurchaseWithAttributes: Failed to get the default analytics provider. Double check your [Analytics] configuration in your INI"));
+	}
+}
+
 void UAnalyticsBlueprintLibrary::RecordCurrencyPurchase(const FString& GameCurrencyType, int32 GameCurrencyAmount, const FString& RealCurrencyType, float RealMoneyCost, const FString& PaymentProvider)
 {
 	TSharedPtr<IAnalyticsProvider> Provider = FAnalytics::Get().GetDefaultConfiguredProvider();
@@ -174,6 +234,19 @@ void UAnalyticsBlueprintLibrary::RecordCurrencyGiven(const FString& GameCurrency
 	}
 }
 
+void UAnalyticsBlueprintLibrary::RecordCurrencyGivenWithAttributes(const FString& GameCurrencyType, int32 GameCurrencyAmount, const TArray<FAnalyticsEventAttr>& Attributes)
+{
+	TSharedPtr<IAnalyticsProvider> Provider = FAnalytics::Get().GetDefaultConfiguredProvider();
+	if (Provider.IsValid())
+	{
+		Provider->RecordCurrencyGiven(GameCurrencyType, GameCurrencyAmount, ConvertAttrs(Attributes));
+	}
+	else
+	{
+		UE_LOG(LogAnalyticsBPLib, Warning, TEXT("RecordCurrencyGivenWithAttributes: Failed to get the default analytics provider. Double check your [Analytics] configuration in your INI"));
+	}
+}
+
 FString UAnalyticsBlueprintLibrary::GetSessionId()
 {
 	TSharedPtr<IAnalyticsProvider> Provider = FAnalytics::Get().GetDefaultConfiguredProvider();
@@ -227,3 +300,120 @@ FAnalyticsEventAttr UAnalyticsBlueprintLibrary::MakeEventAttribute(const FString
 	EventAttr.Value = AttributeValue;
 	return EventAttr;
 }
+
+void UAnalyticsBlueprintLibrary::SetAge(int32 Age)
+{
+	TSharedPtr<IAnalyticsProvider> Provider = FAnalytics::Get().GetDefaultConfiguredProvider();
+	if (Provider.IsValid())
+	{
+		Provider->SetAge(Age);
+	}
+	else
+	{
+		UE_LOG(LogAnalyticsBPLib, Warning, TEXT("SetAge: Failed to get the default analytics provider. Double check your [Analytics] configuration in your INI"));
+	}
+}
+
+void UAnalyticsBlueprintLibrary::SetLocation(const FString& Location)
+{
+	TSharedPtr<IAnalyticsProvider> Provider = FAnalytics::Get().GetDefaultConfiguredProvider();
+	if (Provider.IsValid())
+	{
+		Provider->SetLocation(Location);
+	}
+	else
+	{
+		UE_LOG(LogAnalyticsBPLib, Warning, TEXT("SetLocation: Failed to get the default analytics provider. Double check your [Analytics] configuration in your INI"));
+	}
+}
+
+void UAnalyticsBlueprintLibrary::SetGender(const FString& Gender)
+{
+	TSharedPtr<IAnalyticsProvider> Provider = FAnalytics::Get().GetDefaultConfiguredProvider();
+	if (Provider.IsValid())
+	{
+		Provider->SetGender(Gender);
+	}
+	else
+	{
+		UE_LOG(LogAnalyticsBPLib, Warning, TEXT("SetGender: Failed to get the default analytics provider. Double check your [Analytics] configuration in your INI"));
+	}
+}
+
+void UAnalyticsBlueprintLibrary::SetBuildInfo(const FString& BuildInfo)
+{
+	TSharedPtr<IAnalyticsProvider> Provider = FAnalytics::Get().GetDefaultConfiguredProvider();
+	if (Provider.IsValid())
+	{
+		Provider->SetBuildInfo(BuildInfo);
+	}
+	else
+	{
+		UE_LOG(LogAnalyticsBPLib, Warning, TEXT("SetBuildInfo: Failed to get the default analytics provider. Double check your [Analytics] configuration in your INI"));
+	}
+}
+
+void UAnalyticsBlueprintLibrary::RecordErrorWithAttributes(const FString& Error, const TArray<FAnalyticsEventAttr>& Attributes)
+{
+	TSharedPtr<IAnalyticsProvider> Provider = FAnalytics::Get().GetDefaultConfiguredProvider();
+	if (Provider.IsValid())
+	{
+		Provider->RecordError(Error, ConvertAttrs(Attributes));
+	}
+	else
+	{
+		UE_LOG(LogAnalyticsBPLib, Warning, TEXT("RecordErrorWithAttributes: Failed to get the default analytics provider. Double check your [Analytics] configuration in your INI"));
+	}
+}
+
+void UAnalyticsBlueprintLibrary::RecordError(const FString& Error)
+{
+	TSharedPtr<IAnalyticsProvider> Provider = FAnalytics::Get().GetDefaultConfiguredProvider();
+	if (Provider.IsValid())
+	{
+		Provider->RecordError(Error);
+	}
+	else
+	{
+		UE_LOG(LogAnalyticsBPLib, Warning, TEXT("RecordError: Failed to get the default analytics provider. Double check your [Analytics] configuration in your INI"));
+	}
+}
+
+void UAnalyticsBlueprintLibrary::RecordProgressWithFullHierarchyAndAttributes(const FString& ProgressType, const TArray<FString>& ProgressNames, const TArray<FAnalyticsEventAttr>& Attributes)
+{
+	TSharedPtr<IAnalyticsProvider> Provider = FAnalytics::Get().GetDefaultConfiguredProvider();
+	if (Provider.IsValid())
+	{
+		Provider->RecordProgress(ProgressType, ProgressNames, ConvertAttrs(Attributes));
+	}
+	else
+	{
+		UE_LOG(LogAnalyticsBPLib, Warning, TEXT("RecordProgressWithFullHierarchyAndAttributes: Failed to get the default analytics provider. Double check your [Analytics] configuration in your INI"));
+	}
+}
+
+void UAnalyticsBlueprintLibrary::RecordProgressWithAttributes(const FString& ProgressType, const FString& ProgressName, const TArray<FAnalyticsEventAttr>& Attributes)
+{
+	TSharedPtr<IAnalyticsProvider> Provider = FAnalytics::Get().GetDefaultConfiguredProvider();
+	if (Provider.IsValid())
+	{
+		Provider->RecordProgress(ProgressType, ProgressName, ConvertAttrs(Attributes));
+	}
+	else
+	{
+		UE_LOG(LogAnalyticsBPLib, Warning, TEXT("RecordProgressWithAttributes: Failed to get the default analytics provider. Double check your [Analytics] configuration in your INI"));
+	}
+}
+
+void UAnalyticsBlueprintLibrary::RecordProgress(const FString& ProgressType, const FString& ProgressName)
+{
+	TSharedPtr<IAnalyticsProvider> Provider = FAnalytics::Get().GetDefaultConfiguredProvider();
+	if (Provider.IsValid())
+	{
+		Provider->RecordProgress(ProgressType, ProgressName);
+	}
+	else
+	{
+		UE_LOG(LogAnalyticsBPLib, Warning, TEXT("RecordProgress: Failed to get the default analytics provider. Double check your [Analytics] configuration in your INI"));
+	}
+}
diff --git a/Engine/Source/Runtime/Analytics/Analytics/Public/Interfaces/IAnalyticsProvider.h b/Engine/Source/Runtime/Analytics/Analytics/Public/Interfaces/IAnalyticsProvider.h
--- a/Engine/Source/Runtime/Analytics/Analytics/Public/Interfaces/IAnalyticsProvider.h
+++ b/Engine/Source/Runtime/Analytics/Analytics/Public/Interfaces/IAnalyticsProvider.h
@@ -95,10 +95,42 @@ public:
 	 */
 	virtual FString GetUserID() const = 0;
 
+	/**
+	 * Sets a user defined string as the build information/version for the session
+	 */
+	virtual void SetBuildInfo(const FString& InBuildInfo)
+	{
+		RecordEvent(TEXT("BuildInfo"), TEXT("BuildInfo"), InBuildInfo);
+	}
+
+	/**
+	 * Sets the gender the game believes the user is as part of the session
+	 */
+	virtual void SetGender(const FString& InGender)
+	{
+		RecordEvent(TEXT("Gender"), TEXT("Gender"), InGender);
+	}
+
+	/**
+	 * Sets the location the game believes the user is playing in as part of the session
+	 */
+	virtual void SetLocation(const FString& InLocation)
+	{
+		RecordEvent(TEXT("Location"), TEXT("Location"), InLocation);
+	}
+
+	/**
+	 * Sets the location the game believes the user is playing in as part of the session
+	 */
+	virtual void SetAge(const int32 InAge)
+	{
+		RecordEvent(TEXT("Age"), TEXT("Age"), *TTypeToString<int32>::ToString(InAge));
+	}
+
 	/**
 	 * Records a named event with an array of attributes
 	 *
-	 * @param EventName name of hte event
+	 * @param EventName name of the event
 	 * @param ParamArray array of attribute name/value pairs
 	 */
 	virtual void RecordEvent(const FString& EventName, const TArray<FAnalyticsEventAttribute>& Attributes) = 0;
@@ -106,7 +138,7 @@ public:
 	/**
 	 * Overload for RecordEvent that takes no parameters
 	 *
-	 * @param EventName name of hte event
+	 * @param EventName name of the event
 	 */
 	void RecordEvent(const FString& EventName)
 	{
@@ -154,13 +186,43 @@ public:
 	virtual void RecordItemPurchase(const FString& ItemId, const FString& Currency, int PerItemCost, int ItemQuantity)
 	{
 		TArray<FAnalyticsEventAttribute> Params;
-		Params.Add(FAnalyticsEventAttribute(TEXT("ItemId"), ItemId));
 		Params.Add(FAnalyticsEventAttribute(TEXT("Currency"), Currency));
 		Params.Add(FAnalyticsEventAttribute(TEXT("PerItemCost"), PerItemCost));
+		RecordItemPurchase(ItemId, ItemQuantity, Params);
+	}
+
+	/**
+	 * Record an in-game purchase of a an item.
+	 * 
+	 * Note that not all providers support item purchase events. In this case this method
+	 * is equivalent to sending a regular event with name "Item Purchase".
+	 * 
+	 * @param ItemId - the ID of the item, should be registered with the provider first.
+	 * @param ItemQuantity - the number of Items purchased.
+	 * @param EventAttrs - a list of key/value pairs to assign to this event
+	 */
+	virtual void RecordItemPurchase(const FString& ItemId, int ItemQuantity, const TArray<FAnalyticsEventAttribute>& EventAttrs)
+	{
+		TArray<FAnalyticsEventAttribute> Params(EventAttrs);
+		Params.Add(FAnalyticsEventAttribute(TEXT("ItemId"), ItemId));
 		Params.Add(FAnalyticsEventAttribute(TEXT("ItemQuantity"), ItemQuantity));
 		RecordEvent(TEXT("Item Purchase"), Params);
 	}
 
+	/**
+	 * Record an in-game purchase of a an item.
+	 * 
+	 * Note that not all providers support item purchase events. In this case this method
+	 * is equivalent to sending a regular event with name "Item Purchase".
+	 * 
+	 * @param ItemId - the ID of the item, should be registered with the provider first.
+	 * @param ItemQuantity - the number of Items purchased.
+	 */
+	void RecordItemPurchase(const FString& ItemId, int ItemQuantity)
+	{
+		RecordItemPurchase(ItemId, ItemQuantity, TArray<FAnalyticsEventAttribute>());
+	}
+
 	/**
 	 * Record a purchase of in-game currency using real-world money.
 	 * 
@@ -176,14 +238,44 @@ public:
 	virtual void RecordCurrencyPurchase(const FString& GameCurrencyType, int GameCurrencyAmount, const FString& RealCurrencyType, float RealMoneyCost, const FString& PaymentProvider)
 	{
 		TArray<FAnalyticsEventAttribute> Params;
-		Params.Add(FAnalyticsEventAttribute(TEXT("GameCurrencyType"), GameCurrencyType));
-		Params.Add(FAnalyticsEventAttribute(TEXT("GameCurrencyAmount"), GameCurrencyAmount));
 		Params.Add(FAnalyticsEventAttribute(TEXT("RealCurrencyType"), RealCurrencyType));
 		Params.Add(FAnalyticsEventAttribute(TEXT("RealMoneyCost"), RealMoneyCost));
 		Params.Add(FAnalyticsEventAttribute(TEXT("PaymentProvider"), PaymentProvider));
+		RecordCurrencyPurchase(GameCurrencyType, GameCurrencyAmount, Params);
+	}
+
+	/**
+	 * Record a purchase of in-game currency using real-world money.
+	 * 
+	 * Note that not all providers support currency events. In this case this method
+	 * is equivalent to sending a regular event with name "Currency Purchase".
+	 * 
+	 * @param GameCurrencyType - type of in game currency purchased, should be registered with the provider first.
+	 * @param GameCurrencyAmount - amount of in game currency purchased.
+	 * @param EventAttrs - a list of key/value pairs to assign to this event
+	 */
+	virtual void RecordCurrencyPurchase(const FString& GameCurrencyType, int GameCurrencyAmount, const TArray<FAnalyticsEventAttribute>& EventAttrs)
+	{
+		TArray<FAnalyticsEventAttribute> Params(EventAttrs);
+		Params.Add(FAnalyticsEventAttribute(TEXT("GameCurrencyType"), GameCurrencyType));
+		Params.Add(FAnalyticsEventAttribute(TEXT("GameCurrencyAmount"), GameCurrencyAmount));
 		RecordEvent(TEXT("Currency Purchase"), Params);
 	}
 
+	/**
+	 * Record an in-game purchase of a an item.
+	 * 
+	 * Note that not all providers support item purchase events. In this case this method
+	 * is equivalent to sending a regular event with name "Item Purchase".
+	 * 
+	 * @param ItemId - the ID of the item, should be registered with the provider first.
+	 * @param ItemQuantity - the number of Items purchased.
+	 */
+	virtual void RecordCurrencyPurchase(const FString& GameCurrencyType, int GameCurrencyAmount)
+	{
+		RecordCurrencyPurchase(GameCurrencyType, GameCurrencyAmount, TArray<FAnalyticsEventAttribute>());
+	}
+
 	/**
 	 * Record a gift of in-game currency from the game itself.
 	 * 
@@ -195,12 +287,112 @@ public:
 	 */
 	virtual void RecordCurrencyGiven(const FString& GameCurrencyType, int GameCurrencyAmount)
 	{
-		TArray<FAnalyticsEventAttribute> Params;
+		RecordCurrencyGiven(GameCurrencyType, GameCurrencyAmount, TArray<FAnalyticsEventAttribute>());
+	}
+
+	/**
+	 * Record a gift of in-game currency from the game itself.
+	 * 
+	 * Note that not all providers support currency events. In this case this method
+	 * is equivalent to sending a regular event with name "Currency Given".
+	 * 
+	 * @param GameCurrencyType - type of in game currency given, should be registered with the provider first.
+	 * @param GameCurrencyAmount - amount of in game currency given.
+	 * @param EventAttrs - a list of key/value pairs to assign to this event
+	 */
+	virtual void RecordCurrencyGiven(const FString& GameCurrencyType, int GameCurrencyAmount, const TArray<FAnalyticsEventAttribute>& EventAttrs)
+	{
+		TArray<FAnalyticsEventAttribute> Params(EventAttrs);
 		Params.Add(FAnalyticsEventAttribute(TEXT("GameCurrencyType"), GameCurrencyType));
 		Params.Add(FAnalyticsEventAttribute(TEXT("GameCurrencyAmount"), GameCurrencyAmount));
 		RecordEvent(TEXT("Currency Given"), Params);
 	}
 
+	/**
+	 * Records an error that has happened in the game
+	 * 
+	 * Note that not all providers support all events. In this case this method
+	 * is equivalent to sending a regular event with name "Game Error".
+	 *
+	 * @param Error - the error string to record
+	 * @param EventAttrs - a list of key/value pairs to assign to this event
+	 */
+	virtual void RecordError(const FString& Error, const TArray<FAnalyticsEventAttribute>& EventAttrs)
+	{
+		TArray<FAnalyticsEventAttribute> Params(EventAttrs);
+		Params.Add(FAnalyticsEventAttribute(TEXT("Error"), *Error));
+		RecordEvent(TEXT("Game Error"), Params);
+	}
+
+	/**
+	 * Records an error that has happened in the game
+	 * 
+	 * Note that not all providers support all events. In this case this method
+	 * is equivalent to sending a regular event with name "Game Error".
+	 *
+	 * @param Error - the error string to record
+	 */
+	virtual void RecordError(const FString& Error)
+	{
+		RecordError(Error, TArray<FAnalyticsEventAttribute>());
+	}
+
+	/**
+	 * Record a player progression event that has happened in the game
+	 * 
+	 * Note that not all providers support all events. In this case this method
+	 * is equivalent to sending a regular event with name "Progression".
+	 *
+	 * @param Error - the error string to record
+	 */
+	virtual void RecordProgress(const FString& ProgressType, const TArray<FString>& ProgressHierarchy, const TArray<FAnalyticsEventAttribute>& EventAttrs)
+	{
+		TArray<FAnalyticsEventAttribute> Params(EventAttrs);
+		Params.Add(FAnalyticsEventAttribute(TEXT("ProgressType"), *ProgressType));
+		FString Hierarchy;
+		// Build a dotted hierarchy string from the list of hierarchy progress
+		for (int32 Index = 0; Index < ProgressHierarchy.Num(); Index++)
+		{
+			Hierarchy += ProgressHierarchy[Index];
+			if (Index + 1 < ProgressHierarchy.Num())
+			{
+				Hierarchy += TEXT(".");
+			}
+		}
+		Params.Add(FAnalyticsEventAttribute(TEXT("ProgressHierarchy"), *Hierarchy));
+		RecordEvent(TEXT("Progression"), Params);
+	}
+
+	/**
+	 * Record a player progression event that has happened in the game
+	 * 
+	 * Note that not all providers support all events. In this case this method
+	 * is equivalent to sending a regular event with name "Progression".
+	 *
+	 * @param Error - the error string to record
+	 */
+	virtual void RecordProgress(const FString& ProgressType, const FString& ProgressHierarchy)
+	{
+		TArray<FString> Hierarchy;
+		Hierarchy.Add(ProgressHierarchy);
+		RecordProgress(ProgressType, Hierarchy, TArray<FAnalyticsEventAttribute>());
+	}
+
+	/**
+	 * Record a player progression event that has happened in the game
+	 * 
+	 * Note that not all providers support all events. In this case this method
+	 * is equivalent to sending a regular event with name "Progression".
+	 *
+	 * @param Error - the error string to record
+	 */
+	virtual void RecordProgress(const FString& ProgressType, const FString& ProgressHierarchy, const TArray<FAnalyticsEventAttribute>& EventAttrs)
+	{
+		TArray<FString> Hierarchy;
+		Hierarchy.Add(ProgressHierarchy);
+		RecordProgress(ProgressType, Hierarchy, EventAttrs);
+	}
+
 	virtual ~IAnalyticsProvider() {}
 };
 
