Interagrated Change from CL#2541361

UE4-Orion to UE4

diff --git a/Engine/Source/Runtime/Analytics/AnalyticsET/Private/AnalyticsET.cpp b/Engine/Source/Runtime/Analytics/AnalyticsET/Private/AnalyticsET.cpp
--- a/Engine/Source/Runtime/Analytics/AnalyticsET/Private/AnalyticsET.cpp
+++ b/Engine/Source/Runtime/Analytics/AnalyticsET/Private/AnalyticsET.cpp
@@ -17,7 +17,8 @@ IMPLEMENT_MODULE( FAnalyticsET, AnalyticsET );
 
 class FAnalyticsProviderET : 
 	public IAnalyticsProvider,
-	public FTickerObjectBase
+	public FTickerObjectBase,
+	public TSharedFromThis<FAnalyticsProviderET>
 {
 public:
 	FAnalyticsProviderET(const FAnalyticsET::Config& ConfigValues);
@@ -70,6 +71,8 @@ private:
 	bool bShouldCacheEvents;
 	/** Current countdown timer to keep track of MaxCachedElapsedTime push */
 	float FlushEventsCountdown;
+	/** Track destructing for unbinding callbacks when firing events at shutdown */
+	bool bInDestructor;
 	/**
 	 * Analytics event entry to be cached
 	 */
@@ -156,6 +159,7 @@ FAnalyticsProviderET::FAnalyticsProviderET(const FAnalyticsET::Config& ConfigVal
 	, MaxCachedElapsedTime(60.0f)
 	, bShouldCacheEvents(!FParse::Param(FCommandLine::Get(), TEXT("ANALYTICSDISABLECACHING")))
 	, FlushEventsCountdown(MaxCachedElapsedTime)
+	, bInDestructor(false)
 {
 	// if we are not caching events, we are operating in debug mode. Turn on super-verbose analytics logging
 	if (!bShouldCacheEvents)
@@ -237,6 +241,7 @@ bool FAnalyticsProviderET::Tick(float DeltaSeconds)
 FAnalyticsProviderET::~FAnalyticsProviderET()
 {
 	UE_LOG(LogAnalytics, Verbose, TEXT("[%s] Destroying ET Analytics provider"), *APIKey);
+	bInDestructor = true;
 	EndSession();
 }
 
@@ -357,8 +362,10 @@ void FAnalyticsProviderET::FlushEvents()
 		HttpRequest->SetURL(APIServer + URLPath);
 		HttpRequest->SetVerb(TEXT("POST"));
 		HttpRequest->SetContentAsString(Payload);
-		HttpRequest->OnProcessRequestComplete().BindRaw(this, &FAnalyticsProviderET::EventRequestComplete);
-
+		if (!bInDestructor)
+		{
+			HttpRequest->OnProcessRequestComplete().BindSP(this, &FAnalyticsProviderET::EventRequestComplete);
+		}
  		HttpRequest->ProcessRequest();
 
 		if (UseDataRouter)
@@ -384,7 +391,10 @@ void FAnalyticsProviderET::FlushEvents()
 
 			HttpRequest->SetVerb(TEXT("POST"));
 			HttpRequest->SetContentAsString(Payload);
-			HttpRequest->OnProcessRequestComplete().BindRaw(this, &FAnalyticsProviderET::EventRequestCompleteDataRouter);
+			if (!bInDestructor)
+			{
+				HttpRequest->OnProcessRequestComplete().BindSP(this, &FAnalyticsProviderET::EventRequestCompleteDataRouter);
+			}
 			HttpRequest->ProcessRequest();
 		}
 
