Reverting change to analytics callbacks. We have fixed this problem by using BindSP instead of BindRaw so callbacks shouldn't be made on deleted objects anymore.

diff --git a/Engine/Source/Runtime/Analytics/AnalyticsET/Private/AnalyticsET.cpp b/Engine/Source/Runtime/Analytics/AnalyticsET/Private/AnalyticsET.cpp
--- a/Engine/Source/Runtime/Analytics/AnalyticsET/Private/AnalyticsET.cpp
+++ b/Engine/Source/Runtime/Analytics/AnalyticsET/Private/AnalyticsET.cpp
@@ -100,10 +100,14 @@ private:
 	FCriticalSection CachedEventsCS;
 
 	/**
-	 * Called when an event Http request completes
+	 * Delegate called when an event Http request completes
 	 */
-	static void LogRequestComplete(bool bIsDataRouter, const FString& InAPIKey, FHttpRequestPtr HttpRequest, FHttpResponsePtr HttpResponse, bool bSucceeded);
+	void EventRequestComplete(FHttpRequestPtr HttpRequest, FHttpResponsePtr HttpResponse, bool bSucceeded);
 
+	/**
+	 * Delegate called when an event Http request completes (for DataRouter)
+	 */
+	void EventRequestCompleteDataRouter(FHttpRequestPtr HttpRequest, FHttpResponsePtr HttpResponse, bool bSucceeded);
 };
 
 void FAnalyticsET::StartupModule()
@@ -287,8 +291,6 @@ void FAnalyticsProviderET::EndSession()
 
 void FAnalyticsProviderET::FlushEvents()
 {
-	// So we can capture a copy of APIKey to the finished calbacks futher down.
-	const FString& APIKeyRef = APIKey;
 	// Make sure we don't try to flush too many times. When we are not caching events it's possible this can be called when there are no events in the array.
 	if (CachedEvents.Num() == 0)
 	{
@@ -360,13 +362,10 @@ void FAnalyticsProviderET::FlushEvents()
 		HttpRequest->SetURL(APIServer + URLPath);
 		HttpRequest->SetVerb(TEXT("POST"));
 		HttpRequest->SetContentAsString(Payload);
-		HttpRequest->OnProcessRequestComplete().BindLambda([APIKeyRef] (FHttpRequestPtr InHttpRequest, FHttpResponsePtr InHttpResponse, bool bSucceeded)
+		if (!bInDestructor)
 		{
-			LogRequestComplete(false, APIKeyRef, InHttpRequest, InHttpResponse, bSucceeded);
-		});
-
-		UE_LOG(LogAnalytics, Display, TEXT("[%s] AnalyticsET URL:%s. Payload:%s"), *APIKey, *URLPath, *Payload);
-
+			HttpRequest->OnProcessRequestComplete().BindSP(this, &FAnalyticsProviderET::EventRequestComplete);
+		}
  		HttpRequest->ProcessRequest();
 
 		if (UseDataRouter)
@@ -392,11 +391,10 @@ void FAnalyticsProviderET::FlushEvents()
 
 			HttpRequest->SetVerb(TEXT("POST"));
 			HttpRequest->SetContentAsString(Payload);
-			HttpRequest->OnProcessRequestComplete().BindLambda([APIKeyRef] (FHttpRequestPtr InHttpRequest, FHttpResponsePtr InHttpResponse, bool bSucceeded)
+			if (!bInDestructor)
 			{
-				LogRequestComplete(true, APIKeyRef, InHttpRequest, InHttpResponse, bSucceeded);
-			});
-
+				HttpRequest->OnProcessRequestComplete().BindSP(this, &FAnalyticsProviderET::EventRequestCompleteDataRouter);
+			}
 			HttpRequest->ProcessRequest();
 		}
 
@@ -455,14 +453,26 @@ void FAnalyticsProviderET::RecordEvent(const FString& EventName, const TArray<FA
 }
 
 
-void FAnalyticsProviderET::LogRequestComplete(bool bIsDataRouter, const FString& APIKey, FHttpRequestPtr HttpRequest, FHttpResponsePtr HttpResponse, bool bSucceeded)
+void FAnalyticsProviderET::EventRequestComplete(FHttpRequestPtr HttpRequest, FHttpResponsePtr HttpResponse, bool bSucceeded)
 {
 	if (bSucceeded && HttpResponse.IsValid())
 	{
-		UE_LOG(LogAnalytics, VeryVerbose, TEXT("[%s] ET%s response for [%s]. Code: %d. Payload: %s"), *APIKey, bIsDataRouter?TEXT(" (DataRouter)"):TEXT(""), *HttpRequest->GetURL(), HttpResponse->GetResponseCode(), *HttpResponse->GetContentAsString());
+		UE_LOG(LogAnalytics, VeryVerbose, TEXT("[%s] ET response for [%s]. Code: %d. Payload: %s"), *APIKey, *HttpRequest->GetURL(), HttpResponse->GetResponseCode(), *HttpResponse->GetContentAsString());
 	}
 	else
 	{
-		UE_LOG(LogAnalytics, VeryVerbose, TEXT("[%s] ET%s response for [%s]. No response"), *APIKey, bIsDataRouter?TEXT(" (DataRouter)"):TEXT(""), *HttpRequest->GetURL());
+		UE_LOG(LogAnalytics, VeryVerbose, TEXT("[%s] ET response for [%s]. No response"), *APIKey, *HttpRequest->GetURL());
 	}
 }
+
+void FAnalyticsProviderET::EventRequestCompleteDataRouter(FHttpRequestPtr HttpRequest, FHttpResponsePtr HttpResponse, bool bSucceeded)
+{
+	if (bSucceeded && HttpResponse.IsValid())
+	{
+		UE_LOG(LogAnalytics, VeryVerbose, TEXT("[%s] ET (DataRouter) response for [%s]. Code: %d. Payload: %s"), *APIKey, *HttpRequest->GetURL(), HttpResponse->GetResponseCode(), *HttpResponse->GetContentAsString());
+	}
+	else
+	{
+		UE_LOG(LogAnalytics, VeryVerbose, TEXT("[%s] ET (DataRouter) response for [%s]. No response"), *APIKey, *HttpRequest->GetURL());
+	}
+}
\ No newline at end of file
