UE-17003 - Removed data router code from AnalyticsET.
FYI, Darren.

diff --git a/Engine/Source/Runtime/Analytics/AnalyticsET/Private/AnalyticsET.cpp b/Engine/Source/Runtime/Analytics/AnalyticsET/Private/AnalyticsET.cpp
--- a/Engine/Source/Runtime/Analytics/AnalyticsET/Private/AnalyticsET.cpp
+++ b/Engine/Source/Runtime/Analytics/AnalyticsET/Private/AnalyticsET.cpp
@@ -59,10 +59,6 @@ private:
 	FString BuildType;
 	/** The AppVersion passed to ET. */
 	FString AppVersion;
-	/** True if we are sending to the data router*/
-	bool UseDataRouter;
-	/** The URL to which uploads are sent when using the data router*/
-	FString DataRouterUploadURL;
 	/** Max number of analytics events to cache before pushing to server */
 	const int32 MaxCachedNumEvents;
 	/** Max time that can elapse before pushing cached events to server */
@@ -103,11 +99,6 @@ private:
 	 * Delegate called when an event Http request completes
 	 */
 	void EventRequestComplete(FHttpRequestPtr HttpRequest, FHttpResponsePtr HttpResponse, bool bSucceeded);
-
-	/**
-	 * Delegate called when an event Http request completes (for DataRouter)
-	 */
-	void EventRequestCompleteDataRouter(FHttpRequestPtr HttpRequest, FHttpResponsePtr HttpResponse, bool bSucceeded);
 };
 
 void FAnalyticsET::StartupModule()
@@ -128,8 +119,6 @@ TSharedPtr<IAnalyticsProvider> FAnalyticsET::CreateAnalyticsProvider(const FAnal
 		ConfigValues.APIKeyET = GetConfigValue.Execute(Config::GetKeyNameForAPIKey(), true);
 		ConfigValues.APIServerET = GetConfigValue.Execute(Config::GetKeyNameForAPIServer(), false);
 		ConfigValues.AppVersionET = GetConfigValue.Execute(Config::GetKeyNameForAppVersion(), false);
-		ConfigValues.UseDataRouterET = GetConfigValue.Execute(Config::GetKeyNameForUseDataRouter(), false);
-		ConfigValues.DataRouterUploadURLET = GetConfigValue.Execute(Config::GetKeyNameForDataRouterUploadURL(), false);
 		return CreateAnalyticsProvider(ConfigValues);
 	}
 	else
@@ -175,14 +164,6 @@ FAnalyticsProviderET::FAnalyticsProviderET(const FAnalyticsET::Config& ConfigVal
 		? FAnalyticsET::Config::GetDefaultAPIServer()
 		: ConfigValues.APIServerET;
 
-	// allow the DataRouterUploadURL value to be empty and use defaults.
-	DataRouterUploadURL = ConfigValues.DataRouterUploadURLET.IsEmpty()
-		? FAnalyticsET::Config::GetDefaultDataRouterUploadURL()
-		: ConfigValues.DataRouterUploadURLET;
-
-	// determine if we are using the data router
-	UseDataRouter = FCString::ToBool(*ConfigValues.UseDataRouterET);
-
 	// default to GEngineVersion if one is not provided, append GEngineVersion otherwise.
 	FString ConfigAppVersion = ConfigValues.AppVersionET;
 	// Allow the cmdline to force a specific AppVersion so it can be set dynamically.
@@ -192,10 +173,6 @@ FAnalyticsProviderET::FAnalyticsProviderET(const FAnalyticsET::Config& ConfigVal
 		: ConfigAppVersion.Replace(TEXT("%VERSION%"), *GEngineVersion.ToString(), ESearchCase::CaseSensitive);
 
 	UE_LOG(LogAnalytics, Log, TEXT("[%s] APIServer = %s. AppVersion = %s"), *APIKey, *APIServer, *AppVersion);
-	if (UseDataRouter)
-	{
-		UE_LOG(LogAnalytics, Log, TEXT("[%s] DataRouterUploadURL = %s. AppVersion = %s"), *APIKey, *DataRouterUploadURL, *AppVersion);
-	}
 
 	// cache the build type string
 	FAnalytics::BuildType BuildTypeEnum = FAnalytics::Get().GetBuildType();
@@ -309,11 +286,6 @@ void FAnalyticsProviderET::FlushEvents()
 
 		TSharedRef< TJsonWriter<TCHAR, TCondensedJsonPrintPolicy<TCHAR> > > JsonWriter = TJsonWriterFactory<TCHAR, TCondensedJsonPrintPolicy<TCHAR> >::Create(&Payload);
 		JsonWriter->WriteObjectStart();
-		if (UseDataRouter)
-		{
-			JsonWriter->WriteValue(TEXT("SessionID"), SessionID);
-			JsonWriter->WriteValue(TEXT("UserID"), UserID);
-		}
 		JsonWriter->WriteArrayStart(TEXT("Events"));
 		for (int32 EventIdx = 0; EventIdx < CachedEvents.Num(); EventIdx++)
 		{
@@ -368,36 +340,6 @@ void FAnalyticsProviderET::FlushEvents()
 		}
  		HttpRequest->ProcessRequest();
 
-		if (UseDataRouter)
-		{
-			// If we're using the DataRouter backend, then submit the same request to the new DataRouter backend
-			// NOTE - This branch is temp, and we will eventually use the DataRouter path exclusively
-
-			// Create/send Http request for an event
-			HttpRequest = FHttpModule::Get().CreateRequest();
-			HttpRequest->SetHeader(TEXT("Content-Type"), TEXT("application/json; charset=utf-8"));
-
-			// TODO need agent here??
-			HttpRequest->SetURL(
-				FString::Printf(TEXT("%s?SessionID=%s&AppID=%s&AppVersion=%s&UserID=%s&IsEditor=%s&AppEnvironment=%s&UploadType=eteventstream"),
-				*DataRouterUploadURL,
-				*FGenericPlatformHttp::UrlEncode(SessionID),
-				*FGenericPlatformHttp::UrlEncode(APIKey), 
-				*FGenericPlatformHttp::UrlEncode(AppVersion),
-				*FGenericPlatformHttp::UrlEncode(UserID),
-				*FGenericPlatformHttp::UrlEncode(FString::FromInt(GIsEditor)),
-				*FGenericPlatformHttp::UrlEncode(BuildType)
-			));
-
-			HttpRequest->SetVerb(TEXT("POST"));
-			HttpRequest->SetContentAsString(Payload);
-			if (!bInDestructor)
-			{
-				HttpRequest->OnProcessRequestComplete().BindSP(this, &FAnalyticsProviderET::EventRequestCompleteDataRouter);
-			}
-			HttpRequest->ProcessRequest();
-		}
-
 		FlushEventsCountdown = MaxCachedElapsedTime;
 		CachedEvents.Empty();
 	}
@@ -464,15 +406,3 @@ void FAnalyticsProviderET::EventRequestComplete(FHttpRequestPtr HttpRequest, FHt
 		UE_LOG(LogAnalytics, VeryVerbose, TEXT("[%s] ET response for [%s]. No response"), *APIKey, *HttpRequest->GetURL());
 	}
 }
-
-void FAnalyticsProviderET::EventRequestCompleteDataRouter(FHttpRequestPtr HttpRequest, FHttpResponsePtr HttpResponse, bool bSucceeded)
-{
-	if (bSucceeded && HttpResponse.IsValid())
-	{
-		UE_LOG(LogAnalytics, VeryVerbose, TEXT("[%s] ET (DataRouter) response for [%s]. Code: %d. Payload: %s"), *APIKey, *HttpRequest->GetURL(), HttpResponse->GetResponseCode(), *HttpResponse->GetContentAsString());
-	}
-	else
-	{
-		UE_LOG(LogAnalytics, VeryVerbose, TEXT("[%s] ET (DataRouter) response for [%s]. No response"), *APIKey, *HttpRequest->GetURL());
-	}
-}
\ No newline at end of file
diff --git a/Engine/Source/Runtime/Analytics/AnalyticsET/Public/AnalyticsET.h b/Engine/Source/Runtime/Analytics/AnalyticsET/Public/AnalyticsET.h
--- a/Engine/Source/Runtime/Analytics/AnalyticsET/Public/AnalyticsET.h
+++ b/Engine/Source/Runtime/Analytics/AnalyticsET/Public/AnalyticsET.h
@@ -44,10 +44,6 @@ public:
 		 * If you supply your own Version string, occurrences of "%VERSION%" are replaced with GEngineVersion. ie, -AnalyticsAppVersion=MyCustomID-%VERSION%.
 		 */
 		FString AppVersionET;
-		/** UseDataRouter - flag indicating that we should use the new datarouter system. */
-		FString UseDataRouterET;
-		/** DataRouterUploadURLET - data router upload URL - Defaults if empty to GetDefaultDataRouterUploadURL. */
-		FString DataRouterUploadURLET;
 
 		/** KeyName required for APIKey configuration. */
 		static FString GetKeyNameForAPIKey() { return TEXT("APIKeyET"); }
@@ -57,12 +53,6 @@ public:
 		static FString GetKeyNameForAppVersion() { return TEXT("AppVersionET"); }
 		/** Default value if no APIServer configuration is provided. */
 		static FString GetDefaultAPIServer() { return TEXT("http://devonline-02:/ETAP/"); }
-		/** Default value if no DataRouterServer configuration is provided. */
-		static FString GetDefaultDataRouterUploadURL() { return TEXT("https://datarouter-public-service-gamedev.ol.epicgames.net/datarouter/api/v1/public/data"); }
-		/** KeyName indicating if we should use the data router. */
-		static FString GetKeyNameForUseDataRouter() { return TEXT("UseDataRouterET"); }
-		/** KeyName for the data router upload URL. */
-		static FString GetKeyNameForDataRouterUploadURL() { return TEXT("DataRouterUploadURLET"); }
 	};
 
 	//--------------------------------------------------------------------------
