Analytics: ability to distinguish between different environments (OP-777).

- Domain name is now sent as part of the aggregated endpoint stats.

diff --git a/Engine/Source/Runtime/Analytics/AnalyticsET/Private/HttpServiceTracker.cpp b/Engine/Source/Runtime/Analytics/AnalyticsET/Private/HttpServiceTracker.cpp
--- a/Engine/Source/Runtime/Analytics/AnalyticsET/Private/HttpServiceTracker.cpp
+++ b/Engine/Source/Runtime/Analytics/AnalyticsET/Private/HttpServiceTracker.cpp
@@ -19,6 +19,7 @@ bool FHttpServiceTracker::Tick(float DeltaTime)
 			for (const auto& MetricsMapPair : EndpointMetricsMap)
 			{
 				Attrs.Reset();
+				Attrs.Emplace(TEXT("DomainName"), MetricsMapPair.Value.LastAnalyticsName);
 				Attrs.Emplace(TEXT("FailCount"), MetricsMapPair.Value.FailCount);
 				Attrs.Emplace(TEXT("SuccessCount"), MetricsMapPair.Value.SuccessCount);
 				// We may have had no successful requests, so these values would be undefined.
@@ -78,6 +79,40 @@ bool FHttpServiceTracker::EndpointMetrics::IsSuccessfulResponse(int32 ResponseCo
 	return ResponseCode >= 200 && ResponseCode < 400;
 }
 
+namespace
+{
+	/**
+	 * @brief Returns name of the endpoint for analytics.
+	 *
+	 * @param FullURL actual URL used
+	 * @return Name for analytics (currently this is a domain name).
+	 */
+	FString GetAnalyticsName(const FString & FullURL)
+	{
+		// use the first part of address
+		int DomainNameBegin = FullURL.Find(TEXT("://"), ESearchCase::CaseSensitive);
+
+		if (DomainNameBegin == INDEX_NONE)
+		{
+			UE_LOG(LogAnalytics, Warning, TEXT("Could not find protocol in URL '%s', analytics name will likely be incorrect"), *FullURL);
+			return FullURL;
+		}
+
+		DomainNameBegin += 3;	// length of "://"
+
+		int DomainNameEnd = FullURL.Find(TEXT("/"), ESearchCase::CaseSensitive, ESearchDir::FromStart, DomainNameBegin);
+
+		if (DomainNameEnd == INDEX_NONE || DomainNameEnd <= DomainNameBegin)
+		{
+			UE_LOG(LogAnalytics, Warning, TEXT("Could not determine domain name in URL '%s', analytics name will likely be incorrect"), *FullURL);
+			return FullURL;
+		}
+
+		return FullURL.Mid(DomainNameBegin, DomainNameEnd - DomainNameBegin);
+	}
+}
+
+
 void FHttpServiceTracker::EndpointMetrics::TrackRequest(const FHttpRequestPtr& HttpRequest)
 {
 	if(HttpRequest.IsValid())
@@ -111,6 +146,15 @@ void FHttpServiceTracker::EndpointMetrics::TrackRequest(const FHttpRequestPtr& H
 			// sum download rate for average calc
 			DownloadBytesFailTotal += DownloadBytes;
 		}
+
+		FString AnalyticsName = GetAnalyticsName(HttpRequest->GetURL());
+		if (LastAnalyticsName.Len() > 0 && AnalyticsName != LastAnalyticsName)
+		{
+			UE_LOG(LogAnalytics, Warning, TEXT("Endpoint analytics name has changed from '%s' to '%s', aggregated stats will be incorrect"),
+				*LastAnalyticsName, *AnalyticsName);
+		}
+
+		LastAnalyticsName = AnalyticsName;
 	}
 }
 
diff --git a/Engine/Source/Runtime/Analytics/AnalyticsET/Public/HttpServiceTracker.h b/Engine/Source/Runtime/Analytics/AnalyticsET/Public/HttpServiceTracker.h
--- a/Engine/Source/Runtime/Analytics/AnalyticsET/Public/HttpServiceTracker.h
+++ b/Engine/Source/Runtime/Analytics/AnalyticsET/Public/HttpServiceTracker.h
@@ -79,6 +79,7 @@ private:
 		int32 SuccessCount;
 		int32 FailCount;
 		TMap<int32, int32> ResponseCodes;
+		FString LastAnalyticsName;
 	};
 
 	/** Mapping of a service endpoint to summary metrics about it. */
