Moving LogAnalytics define into header, as it's used to two cpp files within the same project which causes issues on unity

diff --git a/Engine/Source/Runtime/Analytics/AnalyticsET/Private/AnalyticsET.cpp b/Engine/Source/Runtime/Analytics/AnalyticsET/Private/AnalyticsET.cpp
--- a/Engine/Source/Runtime/Analytics/AnalyticsET/Private/AnalyticsET.cpp
+++ b/Engine/Source/Runtime/Analytics/AnalyticsET/Private/AnalyticsET.cpp
@@ -11,8 +11,6 @@
 #include "Runtime/Online/HTTP/Public/Http.h"
 #include "EngineVersion.h"
 
-DEFINE_LOG_CATEGORY_STATIC(LogAnalytics, Display, All);
-
 IMPLEMENT_MODULE( FAnalyticsET, AnalyticsET );
 
 class FAnalyticsProviderET : 
diff --git a/Engine/Source/Runtime/Analytics/AnalyticsET/Public/AnalyticsET.h b/Engine/Source/Runtime/Analytics/AnalyticsET/Public/AnalyticsET.h
--- a/Engine/Source/Runtime/Analytics/AnalyticsET/Public/AnalyticsET.h
+++ b/Engine/Source/Runtime/Analytics/AnalyticsET/Public/AnalyticsET.h
@@ -5,6 +5,8 @@
 #include "IAnalyticsProviderModule.h"
 #include "Core.h"
 
+DEFINE_LOG_CATEGORY_STATIC(LogAnalytics, Display, All);
+
 class IAnalyticsProvider;
 
 /**
