No longer require an active session to set the SessionID, since we can actually send events without starting a session anyway.

diff --git a/Engine/Source/Runtime/Analytics/AnalyticsET/Private/AnalyticsET.cpp b/Engine/Source/Runtime/Analytics/AnalyticsET/Private/AnalyticsET.cpp
--- a/Engine/Source/Runtime/Analytics/AnalyticsET/Private/AnalyticsET.cpp
+++ b/Engine/Source/Runtime/Analytics/AnalyticsET/Private/AnalyticsET.cpp
@@ -369,13 +369,9 @@ FString FAnalyticsProviderET::GetSessionID() const
 
 bool FAnalyticsProviderET::SetSessionID(const FString& InSessionID)
 {
-	if (bSessionInProgress)
-	{
-		SessionID = InSessionID;
-		UE_LOG(LogAnalytics, Log, TEXT("[%s] Forcing SessionID to %s."), *APIKey, *SessionID);
-		return true;
-	}
-	return false;
+	SessionID = InSessionID;
+	UE_LOG(LogAnalytics, Log, TEXT("[%s] Forcing SessionID to %s."), *APIKey, *SessionID);
+	return true;
 }
 
 /** Helper to log any ET event. Used by all the LogXXX functions. */
