Dmitry's fix for the template/include order compilation issue

diff --git a/Engine/Source/Runtime/Analytics/AnalyticsVisualEditing/Classes/AnalyticsSettings.h b/Engine/Source/Runtime/Analytics/AnalyticsVisualEditing/Classes/AnalyticsSettings.h
--- a/Engine/Source/Runtime/Analytics/AnalyticsVisualEditing/Classes/AnalyticsSettings.h
+++ b/Engine/Source/Runtime/Analytics/AnalyticsVisualEditing/Classes/AnalyticsSettings.h
@@ -2,6 +2,7 @@
 
 #pragma once
 
+#include "CoreUObject.h"
 #include "Engine/DeveloperSettings.h"
 #include "AnalyticsSettings.generated.h"
 
