QoS metrics: add generic module (OP-1132).

(Merging CLs 2688852, 2689232 to the main branch)

diff --git a/Engine/Source/Runtime/Analytics/Analytics/Public/Analytics.h b/Engine/Source/Runtime/Analytics/Analytics/Public/Analytics.h
--- a/Engine/Source/Runtime/Analytics/Analytics/Public/Analytics.h
+++ b/Engine/Source/Runtime/Analytics/Analytics/Public/Analytics.h
@@ -117,6 +117,21 @@ public:
 		Release,
 	};
 
+	static const TCHAR* ToString(BuildType Type)
+	{
+		switch (Type)
+		{
+			case Development:		return TEXT("Development");
+			case Test:				return TEXT("Test");
+			case Debug:				return TEXT("Debug");
+			case Release:			return TEXT("Release");
+			default:
+				break;
+		}
+
+		return TEXT("Undefined");
+	}
+
 	/** Get the analytics build type. Generally used to determine the keys to use to configure an analytics provider. */
 	virtual BuildType GetBuildType() const;
 
diff --git a/Engine/Source/Runtime/Analytics/QoSReporter/Private/QoSReporter.cpp b/Engine/Source/Runtime/Analytics/QoSReporter/Private/QoSReporter.cpp
new file mode 100644
--- /dev/null
+++ b/Engine/Source/Runtime/Analytics/QoSReporter/Private/QoSReporter.cpp
@@ -0,0 +1,100 @@
+// Copyright 1998-2015 Epic Games, Inc. All Rights Reserved.
+
+#include "QoSReporterPrivatePCH.h"
+
+#include "Core.h"
+#include "Analytics.h"
+#include "EngineVersion.h"
+#include "QoSReporter.h"
+
+#ifndef WITH_QOSREPORTER
+	#error "WITH_QOSREPORTER should be defined in Build.cs file"
+#endif
+
+bool FQoSReporter::bIsInitialized;
+TSharedPtr<IAnalyticsProvider> FQoSReporter::Analytics;
+
+/**
+* External code should bind this delegate if QoS reporting is desired,
+* preferably in private code that won't be redistributed.
+*/
+QOSREPORTER_API FAnalytics::FProviderConfigurationDelegate& GetQoSOverrideConfigDelegate()
+{
+	static FAnalytics::FProviderConfigurationDelegate Delegate;
+	return Delegate;
+}
+
+
+/**
+* Get analytics pointer
+*/
+IAnalyticsProvider& FQoSReporter::GetProvider()
+{
+	checkf(bIsInitialized && IsAvailable(), TEXT("FQoSReporter::GetProvider called outside of Initialize/Shutdown."));
+
+	return *Analytics.Get();
+}
+
+void FQoSReporter::Initialize()
+{
+	checkf(!bIsInitialized, TEXT("FQoSReporter::Initialize called more than once."));
+
+	// Setup some default engine analytics if there is nothing custom bound
+	FAnalytics::FProviderConfigurationDelegate DefaultEngineAnalyticsConfig;
+	DefaultEngineAnalyticsConfig.BindLambda(
+		[=](const FString& KeyName, bool bIsValueRequired) -> FString
+	{
+		static TMap<FString, FString> ConfigMap;
+		if (ConfigMap.Num() == 0)
+		{
+			ConfigMap.Add(TEXT("ProviderModuleName"), TEXT("QoSReporter"));
+			ConfigMap.Add(TEXT("APIKeyQoS"), FString::Printf(TEXT("%s.%s"), FApp::GetGameName(), FAnalytics::ToString(FAnalytics::Get().GetBuildType())));
+		}
+
+		// Check for overrides
+		if (GetQoSOverrideConfigDelegate().IsBound())
+		{
+			const FString OverrideValue = GetQoSOverrideConfigDelegate().Execute(KeyName, bIsValueRequired);
+			if (!OverrideValue.IsEmpty())
+			{
+				return OverrideValue;
+			}
+		}
+
+		FString* ConfigValue = ConfigMap.Find(KeyName);
+		return ConfigValue != NULL ? *ConfigValue : TEXT("");
+	});
+
+	// Connect the engine analytics provider (if there is a configuration delegate installed)
+	Analytics = FAnalytics::Get().CreateAnalyticsProvider(
+		FName(*DefaultEngineAnalyticsConfig.Execute(TEXT("ProviderModuleName"), true)),
+		DefaultEngineAnalyticsConfig);
+
+	bIsInitialized = true;
+}
+
+void FQoSReporter::Shutdown()
+{
+	Analytics.Reset();
+	bIsInitialized = false;
+}
+
+double FQoSReporter::ModuleInitializationTime = FPlatformTime::Seconds();
+bool FQoSReporter::bStartupEventReported = false;
+
+void FQoSReporter::ReportStartupCompleteEvent()
+{
+	if (bStartupEventReported || !Analytics.IsValid())
+	{
+		return;
+	}
+
+	double CurrentTime = FPlatformTime::Seconds();
+	double Difference = CurrentTime - ModuleInitializationTime;
+
+	TArray<FAnalyticsEventAttribute> ParamArray;
+	ParamArray.Add(FAnalyticsEventAttribute(EQoSEvents::ToString(EQoSEventParam::StartupTime), Difference));
+	Analytics->RecordEvent(EQoSEvents::ToString(EQoSEventParam::StartupTime), ParamArray);
+
+	bStartupEventReported = true;
+}
diff --git a/Engine/Source/Runtime/Analytics/QoSReporter/Private/QoSReporterModule.cpp b/Engine/Source/Runtime/Analytics/QoSReporter/Private/QoSReporterModule.cpp
new file mode 100644
--- /dev/null
+++ b/Engine/Source/Runtime/Analytics/QoSReporter/Private/QoSReporterModule.cpp
@@ -0,0 +1,240 @@
+// Copyright 1998-2015 Epic Games, Inc. All Rights Reserved.
+
+#include "QoSReporterPrivatePCH.h"
+
+#include "Core.h"
+#include "Json.h"
+#include "Analytics.h"
+#include "Http.h"
+#include "EngineVersion.h"
+#include "QoSReporter.h"
+
+DEFINE_LOG_CATEGORY_STATIC(LogQoSReporter, Display, All);
+
+IMPLEMENT_MODULE(FQoSReporterModule, QoSReporter);
+
+FString FQoSReporterModule::Config::GetDefaultAppVersion()
+{ 
+	return FString::Printf(TEXT("UE4-CL-%d"), GEngineVersion.GetChangelist()); 
+}
+
+FString FQoSReporterModule::Config::GetDefaultAppEnvironment() 
+{
+	return FAnalytics::ToString(FAnalytics::Get().GetBuildType()); 
+}
+
+class FAnalyticsProviderQoSReporter : public IAnalyticsProvider
+{
+public:
+	FAnalyticsProviderQoSReporter(const FQoSReporterModule::Config& ConfigValues);
+
+	/** This provider does not have a concept of sessions */
+	virtual bool StartSession(const TArray<FAnalyticsEventAttribute>& Attributes) override { return true; };
+	/** This provider does not have a concept of sessions */
+	virtual void EndSession() override {};
+	/** This provider is not supposed to send many events, and due to nature of QoS we don't want to cache them */
+	virtual void FlushEvents() override {};
+
+	/** This provider is not using user IDs */
+	virtual void SetUserID(const FString& InUserID) override {};
+	/** This provider is not using user IDs */
+	virtual FString GetUserID() const override { checkf(false, TEXT("FAnalyticsProviderQoSReporter does not use user ids"));  return TEXT("UnknownUserId"); };
+
+	/** This provider does not have a concept of sessions */
+	virtual FString GetSessionID() const override { checkf(false, TEXT("FAnalyticsProviderQoSReporter is not session based"));  return TEXT("UnknownSessionId"); };
+	/** This provider does not have a concept of sessions */
+	virtual bool SetSessionID(const FString& InSessionID) override { return false; };
+
+	virtual void RecordEvent(const FString& EventName, const TArray<FAnalyticsEventAttribute>& Attributes) override;
+	virtual ~FAnalyticsProviderQoSReporter();
+
+	FString GetAPIKey() const { return APIKey; }
+
+private:
+
+	/** API key (also known as "upload type" on data router) */
+	FString APIKey;
+	/** API Server to use (also known as "endpoint"). */
+	FString APIServer;
+	/** The AppVersion to use. */
+	FString AppVersion;
+	/** The AppEnvironment to use. */
+	FString AppEnvironment;
+	/** The upload type to use. */
+	FString UploadType;
+
+	/**
+	 * Delegate called when an event Http request completes
+	 */
+	void EventRequestComplete(FHttpRequestPtr HttpRequest, FHttpResponsePtr HttpResponse, bool bSucceeded);
+};
+
+void FQoSReporterModule::StartupModule()
+{
+}
+
+void FQoSReporterModule::ShutdownModule()
+{
+}
+
+TSharedPtr<IAnalyticsProvider> FQoSReporterModule::CreateAnalyticsProvider(const FAnalytics::FProviderConfigurationDelegate& GetConfigValue) const
+{
+	if (GetConfigValue.IsBound())
+	{
+		Config ConfigValues;
+		ConfigValues.APIServer = GetConfigValue.Execute(Config::GetKeyNameForAPIServer(), true);
+		ConfigValues.APIKey = GetConfigValue.Execute(Config::GetKeyNameForAPIKey(), false);
+		ConfigValues.AppVersion = GetConfigValue.Execute(Config::GetKeyNameForAppVersion(), false);
+		ConfigValues.AppEnvironment = GetConfigValue.Execute(Config::GetKeyNameForAppEnvironment(), false);
+		ConfigValues.UploadType = GetConfigValue.Execute(Config::GetKeyNameForUploadType(), false);
+		return CreateAnalyticsProvider(ConfigValues);
+	}
+	else
+	{
+		UE_LOG(LogQoSReporter, Warning, TEXT("CreateAnalyticsProvider called with an unbound delegate"));
+	}
+	return nullptr;
+}
+
+TSharedPtr<IAnalyticsProvider> FQoSReporterModule::CreateAnalyticsProvider(const Config& ConfigValues) const
+{
+	return TSharedPtr<IAnalyticsProvider>(new FAnalyticsProviderQoSReporter(ConfigValues));
+}
+
+/**
+ * Perform any initialization.
+ */
+FAnalyticsProviderQoSReporter::FAnalyticsProviderQoSReporter(const FQoSReporterModule::Config& ConfigValues)
+{
+	UE_LOG(LogQoSReporter, Verbose, TEXT("Initializing QoS Reporter"));
+
+	APIKey = ConfigValues.APIKey;
+	if (APIKey.IsEmpty())
+	{
+		UE_LOG(LogQoSReporter, Error, TEXT("QoS API key is not configured, no QoS metrics will be reported."));
+	}
+
+	APIServer = ConfigValues.APIServer;
+	if (APIServer.IsEmpty())
+	{
+		UE_LOG(LogQoSReporter, Error, TEXT("QoS API server is not configured, no QoS metrics will be reported."));
+	}
+
+	AppVersion = ConfigValues.AppVersion;
+	if (AppVersion.IsEmpty())
+	{
+		AppVersion = ConfigValues.GetDefaultAppVersion();
+	}
+
+	AppEnvironment = ConfigValues.AppEnvironment;
+	if (AppEnvironment.IsEmpty())
+	{
+		AppEnvironment = ConfigValues.GetDefaultAppEnvironment();
+	}
+
+	UploadType = ConfigValues.UploadType;
+	if (UploadType.IsEmpty())
+	{
+		UploadType = ConfigValues.GetDefaultUploadType();
+	}
+
+	UE_LOG(LogQoSReporter, Log, TEXT("APIKey = '%s'. APIServer = '%s'. AppVersion = '%s'. AppEnvironment = '%s'"), *APIKey, *APIServer, *AppVersion, *AppEnvironment);
+}
+
+FAnalyticsProviderQoSReporter::~FAnalyticsProviderQoSReporter()
+{
+	UE_LOG(LogQoSReporter, Verbose, TEXT("Destroying QoS Reporter"));
+	EndSession();
+}
+
+void FAnalyticsProviderQoSReporter::RecordEvent(const FString& EventName, const TArray<FAnalyticsEventAttribute>& Attributes)
+{
+	if (APIKey.IsEmpty() || APIServer.IsEmpty())
+	{
+		return;
+	}
+
+	// encode params as Json
+	if (ensure(FModuleManager::Get().IsModuleLoaded("HTTP")))
+	{
+		FString Payload;
+
+		FDateTime CurrentTime = FDateTime::UtcNow();
+
+		TSharedRef< TJsonWriter<TCHAR, TCondensedJsonPrintPolicy<TCHAR> > > JsonWriter = TJsonWriterFactory<TCHAR, TCondensedJsonPrintPolicy<TCHAR> >::Create(&Payload);
+		JsonWriter->WriteObjectStart();
+		JsonWriter->WriteArrayStart(TEXT("Events"));
+
+		// write just a single event
+		JsonWriter->WriteObjectStart();
+		JsonWriter->WriteValue(TEXT("EventName"), EventName);
+
+		if (Attributes.Num() > 0)
+		{
+			// optional attributes for this event
+			for (int32 AttrIdx = 0; AttrIdx < Attributes.Num(); AttrIdx++)
+			{
+				const FAnalyticsEventAttribute& Attr = Attributes[AttrIdx];
+				JsonWriter->WriteValue(Attr.AttrName, Attr.AttrValue);
+			}
+		}
+		JsonWriter->WriteObjectEnd();
+
+		JsonWriter->WriteArrayEnd();
+		JsonWriter->WriteObjectEnd();
+		JsonWriter->Close();
+
+		FString URLPath = FString::Printf(TEXT("?AppID=%s&AppVersion=%s&AppEnvironment=%s&UploadType=%s"),
+			*FGenericPlatformHttp::UrlEncode(APIKey),
+			*FGenericPlatformHttp::UrlEncode(AppVersion),
+			*FGenericPlatformHttp::UrlEncode(AppEnvironment),
+			*FGenericPlatformHttp::UrlEncode(UploadType)
+			);
+
+		// Recreate the URLPath for logging because we do not want to escape the parameters when logging.
+		// We cannot simply UrlEncode the entire Path after logging it because UrlEncode(Params) != UrlEncode(Param1) & UrlEncode(Param2) ...
+		UE_LOG(LogQoSReporter, VeryVerbose, TEXT("[%s] QoS URL:%s?AppID=%s&AppVersion=%s&AppEnvironment=%s&UploadType=%s. Payload:%s"),
+			*APIKey,
+			*APIServer,
+			*APIKey,
+			*AppVersion,
+			*AppEnvironment,
+			*UploadType,
+			*Payload);
+
+		// Create/send Http request for an event
+		TSharedRef<IHttpRequest> HttpRequest = FHttpModule::Get().CreateRequest();
+		HttpRequest->SetHeader(TEXT("Content-Type"), TEXT("application/json; charset=utf-8"));
+
+		HttpRequest->SetURL(APIServer + URLPath);
+		HttpRequest->SetVerb(TEXT("POST"));
+		HttpRequest->SetContentAsString(Payload);
+		HttpRequest->OnProcessRequestComplete().BindRaw(this, &FAnalyticsProviderQoSReporter::EventRequestComplete);
+		HttpRequest->ProcessRequest();
+
+	}
+}
+
+void FAnalyticsProviderQoSReporter::EventRequestComplete(FHttpRequestPtr HttpRequest, FHttpResponsePtr HttpResponse, bool bSucceeded)
+{
+	if (bSucceeded && HttpResponse.IsValid())
+	{
+		// normal operation is silent, but any problems are reported as warnings
+		if (EHttpResponseCodes::IsOk(HttpResponse->GetResponseCode()))
+		{
+			UE_LOG(LogQoSReporter, VeryVerbose, TEXT("QoS response for [%s]. Code: %d. Payload: %s"),
+				*HttpRequest->GetURL(), HttpResponse->GetResponseCode(), *HttpResponse->GetContentAsString());
+		}
+		else
+		{
+			UE_LOG(LogQoSReporter, Warning, TEXT("Bad QoS response for [%s] - code: %d. Payload: %s"),
+				*HttpRequest->GetURL(), HttpResponse->GetResponseCode(), *HttpResponse->GetContentAsString());
+		}
+	}
+	else
+	{
+		// if we cannot report QoS metrics this is pretty bad; report at least a warning
+		UE_LOG(LogQoSReporter, Warning, TEXT("QoS response for [%s]. No response"), *HttpRequest->GetURL());
+	}
+}
+
diff --git a/Engine/Source/Runtime/Analytics/QoSReporter/Private/QoSReporterPrivatePCH.h b/Engine/Source/Runtime/Analytics/QoSReporter/Private/QoSReporterPrivatePCH.h
new file mode 100644
--- /dev/null
+++ b/Engine/Source/Runtime/Analytics/QoSReporter/Private/QoSReporterPrivatePCH.h
@@ -0,0 +1,4 @@
+// Copyright 1998-2015 Epic Games, Inc. All Rights Reserved.
+
+// You should place include statements to your module's private header files here.  You only need to
+// add includes for headers that are used in most of your module's source files though.
diff --git a/Engine/Source/Runtime/Analytics/QoSReporter/Public/QoSReporter.h b/Engine/Source/Runtime/Analytics/QoSReporter/Public/QoSReporter.h
new file mode 100644
--- /dev/null
+++ b/Engine/Source/Runtime/Analytics/QoSReporter/Public/QoSReporter.h
@@ -0,0 +1,162 @@
+// Copyright 1998-2015 Epic Games, Inc. All Rights Reserved.
+
+#pragma once
+
+#include "Runtime/Analytics/Analytics/Public/Interfaces/IAnalyticsProviderModule.h"
+#include "Core.h"
+
+class IAnalyticsProvider;
+
+enum class EQoSEventParam : uint32
+{
+	/** Time from initializing the QoS reporter module to the "ready" state of the game/program (e.g. login screen) */
+	StartupTime,
+
+	MAX
+};
+
+namespace EQoSEvents
+{
+	inline const TCHAR * ToString(EQoSEventParam EventParam)
+	{
+		switch (EventParam)
+		{
+			case EQoSEventParam::StartupTime: 
+				return TEXT("StartupTime");
+			default:
+				break;
+		}
+
+		return TEXT("Undefined");
+	}
+}
+
+/**
+ * The public interface to this module
+ */
+class FQoSReporterModule : public IAnalyticsProviderModule
+{
+	//--------------------------------------------------------------------------
+	// Module functionality
+	//--------------------------------------------------------------------------
+public:
+	/**
+	 * Singleton-like access to this module's interface.  This is just for convenience!
+	 * Beware of calling this during the shutdown phase, though.  Your module might have been unloaded already.
+	 *
+	 * @return Returns singleton instance, loading the module on demand if needed
+	 */
+	static inline FQoSReporterModule& Get()
+	{
+		return FModuleManager::LoadModuleChecked< FQoSReporterModule >("QoSReporter");
+	}
+
+	//--------------------------------------------------------------------------
+	// Configuration functionality
+	//--------------------------------------------------------------------------
+public:
+	/** Defines required configuration values for Swrve analytics provider. */
+	struct Config
+	{
+		/** Endpoint (URL) to send events to */
+		FString APIServer;
+		/** API key arbitrary string that identifies the application. */
+		FString APIKey;
+		/** An arbitrary string that identifies the version of the application. */
+		FString AppVersion;
+		/** An arbitrary string that allows to distinguish running environment of a single application (Development, Production, Staging etc). */
+		FString AppEnvironment;
+		/** An arbitrary string that identifies type of events (QoS). */
+		FString UploadType;
+
+		/** Key name for API server (endpoint) */
+		static FString GetKeyNameForAPIServer() { return TEXT("APIServerQoS"); }
+		/** Key name for API key. */
+		static FString GetKeyNameForAPIKey() { return "APIKeyQoS"; }
+		/** Key name for app version. */
+		static FString GetKeyNameForAppVersion() { return "AppVersionQoS"; }
+		/** Key name for app environment. */
+		static FString GetKeyNameForAppEnvironment() { return "AppEnvironmentQoS"; }
+		/** Key name for upload type. */
+		static FString GetKeyNameForUploadType() { return "UploadTypeQoS"; }
+
+		/** Returns default app version. */
+		static FString GetDefaultAppVersion();
+		/** Returns default app environment. */
+		static FString GetDefaultAppEnvironment();
+		/** Returns default upload type. */
+		static FString GetDefaultUploadType() { return TEXT("ue4qosmetrics"); }
+	};
+
+	//--------------------------------------------------------------------------
+	// provider factory functions
+	//--------------------------------------------------------------------------
+public:
+
+	/**
+	 * IAnalyticsProviderModule interface.
+	 * Creates the analytics provider given a configuration delegate.
+	 * The keys required exactly match the field names in the Config object. 
+	 */
+	virtual TSharedPtr<IAnalyticsProvider> CreateAnalyticsProvider(const FAnalytics::FProviderConfigurationDelegate& GetConfigValue) const;
+	
+	/** 
+	 * Construct an analytics provider directly from a config object.
+	 */
+	virtual TSharedPtr<IAnalyticsProvider> CreateAnalyticsProvider(const Config& ConfigValues) const;
+
+private:
+
+	virtual void StartupModule() override;
+	virtual void ShutdownModule() override;
+};
+
+/**
+ * The public interface for the QoS analytics provider singleton.
+ *
+ */
+class FQoSReporter : FNoncopyable
+{
+public:
+	/**
+	 * Return the provider instance. Not valid outside of Initialize/Shutdown calls.
+	 * Note: must check IsAvailable() first else this code will assert if the provider is not valid.
+	 */
+	static IAnalyticsProvider& GetProvider();
+	/** Helper function to determine if the provider is valid. */
+	static bool IsAvailable() { return Analytics.IsValid(); }
+	/** Called to initialize the singleton. */
+	QOSREPORTER_API static void Initialize();
+	/** Called to shut down the singleton */
+	QOSREPORTER_API static void Shutdown();
+
+
+	/**
+	 * These calls implement logic common to all games/programs; they are expected
+	 * to sprinkle calls to them in appropriate place
+	 */
+
+	/**
+	 * This function is expected to be called when a game reached the point of interactivity.
+	 * It is allowed to call this function more than once.
+	 */
+	static void QOSREPORTER_API ReportStartupCompleteEvent();
+
+private:
+
+	/** Whether the module has been initialized. */
+	static bool bIsInitialized;
+
+	/** Chosen analytics provider. */
+	static TSharedPtr<IAnalyticsProvider> Analytics;
+
+	/** 
+	 * Event-specific variables
+	 */
+	/** Used for startup event */
+	static double ModuleInitializationTime;
+
+	/** Whether startup event was reported. */
+	static bool bStartupEventReported;
+
+};
diff --git a/Engine/Source/Runtime/Analytics/QoSReporter/QoSReporter.Build.cs b/Engine/Source/Runtime/Analytics/QoSReporter/QoSReporter.Build.cs
new file mode 100644
--- /dev/null
+++ b/Engine/Source/Runtime/Analytics/QoSReporter/QoSReporter.Build.cs
@@ -0,0 +1,23 @@
+// Copyright 1998-2015 Epic Games, Inc. All Rights Reserved.
+
+namespace UnrealBuildTool.Rules
+{
+	public class QoSReporter : ModuleRules
+	{
+		public QoSReporter(TargetInfo Target)
+		{
+			PublicDependencyModuleNames.Add("Core");
+
+			PrivateDependencyModuleNames.AddRange(
+				new string[]
+				{
+					"Analytics",
+					"HTTP",
+					"Json",
+				}
+			);
+
+			Definitions.Add("WITH_QOSREPORTER=1");
+		}
+	}
+}
