Fixed compilation issue.

diff --git a/Engine/Source/Runtime/Analytics/QoSReporter/Private/QoSReporterModule.cpp b/Engine/Source/Runtime/Analytics/QoSReporter/Private/QoSReporterModule.cpp
--- a/Engine/Source/Runtime/Analytics/QoSReporter/Private/QoSReporterModule.cpp
+++ b/Engine/Source/Runtime/Analytics/QoSReporter/Private/QoSReporterModule.cpp
@@ -15,7 +15,7 @@ IMPLEMENT_MODULE(FQoSReporterModule, QoSReporter);
 
 FString FQoSReporterModule::Config::GetDefaultAppVersion()
 { 
-	return FString::Printf(TEXT("UE4-CL-%d"), GEngineVersion.GetChangelist()); 
+	return FString::Printf(TEXT("UE4-CL-%d"), FEngineVersion::Current().GetChangelist());
 }
 
 FString FQoSReporterModule::Config::GetDefaultAppEnvironment() 
