QoSReporter: add average FPS (point-sampled) reporting.

(Merging CLs 2701668, 2702482, 2703275, 2703344,  2703548 to the main branch)

diff --git a/Engine/Source/Runtime/Analytics/QoSReporter/Private/QoSReporter.cpp b/Engine/Source/Runtime/Analytics/QoSReporter/Private/QoSReporter.cpp
--- a/Engine/Source/Runtime/Analytics/QoSReporter/Private/QoSReporter.cpp
+++ b/Engine/Source/Runtime/Analytics/QoSReporter/Private/QoSReporter.cpp
@@ -39,6 +39,14 @@ void FQoSReporter::Initialize()
 {
 	checkf(!bIsInitialized, TEXT("FQoSReporter::Initialize called more than once."));
 
+	// allow overriding from configs
+	bool bEnabled = true;
+	if (GConfig->GetBool(TEXT("QoSReporter"), TEXT("bEnabled"), bEnabled, GEngineIni) && !bEnabled)
+	{
+		UE_LOG(LogQoSReporter, Verbose, TEXT("QoSReporter disabled by config setting"));
+		return;
+	}
+
 	// Setup some default engine analytics if there is nothing custom bound
 	FAnalytics::FProviderConfigurationDelegate DefaultEngineAnalyticsConfig;
 	DefaultEngineAnalyticsConfig.BindLambda(
@@ -70,6 +78,15 @@ void FQoSReporter::Initialize()
 		FName(*DefaultEngineAnalyticsConfig.Execute(TEXT("ProviderModuleName"), true)),
 		DefaultEngineAnalyticsConfig);
 
+	// check if Configs override the heartbeat interval
+	
+	float ConfigHeartbeatInterval = 0.0;
+	if (GConfig->GetFloat(TEXT("QoSReporter"), TEXT("HeartbeatInterval"), ConfigHeartbeatInterval, GEngineIni))
+	{
+		HeartbeatInterval = ConfigHeartbeatInterval;
+		UE_LOG(LogQoSReporter, Verbose, TEXT("HeartbeatInterval configured to %f from config."), HeartbeatInterval);
+	}
+
 	bIsInitialized = true;
 }
 
@@ -98,3 +115,31 @@ void FQoSReporter::ReportStartupCompleteEvent()
 
 	bStartupEventReported = true;
 }
+
+double FQoSReporter::HeartbeatInterval = 300;
+double FQoSReporter::LastHeartbeatTimestamp = FPlatformTime::Seconds();
+
+void FQoSReporter::Tick()
+{
+	if (!Analytics.IsValid())
+	{
+		return;
+	}
+
+	double CurrentTime = FPlatformTime::Seconds();
+	if (HeartbeatInterval > 0 && CurrentTime - LastHeartbeatTimestamp > HeartbeatInterval)
+	{
+		SendHeartbeat();
+		LastHeartbeatTimestamp = CurrentTime;
+	}
+}
+
+void FQoSReporter::SendHeartbeat()
+{
+	checkf(Analytics.IsValid(), TEXT("SendHeartbeat() should not be called if Analytics provider was not configured"));
+
+	extern ENGINE_API float GAverageFPS;
+	TArray<FAnalyticsEventAttribute> ParamArray;
+	ParamArray.Add(FAnalyticsEventAttribute(TEXT("AverageFPS"), GAverageFPS));
+	Analytics->RecordEvent(EQoSEvents::ToString(EQoSEventParam::Heartbeat), ParamArray);
+}
diff --git a/Engine/Source/Runtime/Analytics/QoSReporter/Private/QoSReporterModule.cpp b/Engine/Source/Runtime/Analytics/QoSReporter/Private/QoSReporterModule.cpp
--- a/Engine/Source/Runtime/Analytics/QoSReporter/Private/QoSReporterModule.cpp
+++ b/Engine/Source/Runtime/Analytics/QoSReporter/Private/QoSReporterModule.cpp
@@ -9,7 +9,7 @@
 #include "EngineVersion.h"
 #include "QoSReporter.h"
 
-DEFINE_LOG_CATEGORY_STATIC(LogQoSReporter, Display, All);
+DEFINE_LOG_CATEGORY(LogQoSReporter);
 
 IMPLEMENT_MODULE(FQoSReporterModule, QoSReporter);
 
diff --git a/Engine/Source/Runtime/Analytics/QoSReporter/Private/QoSReporterPrivatePCH.h b/Engine/Source/Runtime/Analytics/QoSReporter/Private/QoSReporterPrivatePCH.h
--- a/Engine/Source/Runtime/Analytics/QoSReporter/Private/QoSReporterPrivatePCH.h
+++ b/Engine/Source/Runtime/Analytics/QoSReporter/Private/QoSReporterPrivatePCH.h
@@ -2,3 +2,9 @@
 
 // You should place include statements to your module's private header files here.  You only need to
 // add includes for headers that are used in most of your module's source files though.
+
+#pragma once
+
+#include "Core.h"
+
+DECLARE_LOG_CATEGORY_EXTERN(LogQoSReporter, Display, All);
diff --git a/Engine/Source/Runtime/Analytics/QoSReporter/Public/QoSReporter.h b/Engine/Source/Runtime/Analytics/QoSReporter/Public/QoSReporter.h
--- a/Engine/Source/Runtime/Analytics/QoSReporter/Public/QoSReporter.h
+++ b/Engine/Source/Runtime/Analytics/QoSReporter/Public/QoSReporter.h
@@ -12,6 +12,9 @@ enum class EQoSEventParam : uint32
 	/** Time from initializing the QoS reporter module to the "ready" state of the game/program (e.g. login screen) */
 	StartupTime,
 
+	/** Sent regularly, includes a number of attributes */
+	Heartbeat,
+
 	MAX
 };
 
@@ -23,6 +26,8 @@ namespace EQoSEvents
 		{
 			case EQoSEventParam::StartupTime: 
 				return TEXT("StartupTime");
+			case EQoSEventParam::Heartbeat:
+				return TEXT("Heartbeat");
 			default:
 				break;
 		}
@@ -142,15 +147,31 @@ public:
 	 */
 	static void QOSREPORTER_API ReportStartupCompleteEvent();
 
+	/**
+	 * This function is expected to be called periodically to update ongoing tasks.
+	 */
+	static void Tick();
+
 private:
 
+	/**
+	 * Sends heartbeat stats.
+	 */
+	static void SendHeartbeat();
+
 	/** Whether the module has been initialized. */
 	static bool bIsInitialized;
 
 	/** Chosen analytics provider. */
 	static TSharedPtr<IAnalyticsProvider> Analytics;
 
-	/** 
+	/** Heartbeat interval in seconds */
+	static double HeartbeatInterval;
+
+	/** Timestamp when we sent the last heartbeat */
+	static double LastHeartbeatTimestamp;
+
+	/**
 	 * Event-specific variables
 	 */
 	/** Used for startup event */
diff --git a/Engine/Source/Runtime/Analytics/QoSReporter/QoSReporter.Build.cs b/Engine/Source/Runtime/Analytics/QoSReporter/QoSReporter.Build.cs
--- a/Engine/Source/Runtime/Analytics/QoSReporter/QoSReporter.Build.cs
+++ b/Engine/Source/Runtime/Analytics/QoSReporter/QoSReporter.Build.cs
@@ -14,6 +14,7 @@ namespace UnrealBuildTool.Rules
 					"Analytics",
 					"HTTP",
 					"Json",
+					"Engine",	// for GAverageFPS
 				}
 			);
 
