Add support for legacy analytics provider for internal telemetry.

diff --git a/Engine/Source/Runtime/Analytics/AnalyticsET/Private/AnalyticsET.cpp b/Engine/Source/Runtime/Analytics/AnalyticsET/Private/AnalyticsET.cpp
--- a/Engine/Source/Runtime/Analytics/AnalyticsET/Private/AnalyticsET.cpp
+++ b/Engine/Source/Runtime/Analytics/AnalyticsET/Private/AnalyticsET.cpp
@@ -67,6 +67,8 @@ private:
 	float FlushEventsCountdown;
 	/** Track destructing for unbinding callbacks when firing events at shutdown */
 	bool bInDestructor;
+	/** True to use the legacy backend server protocol that uses URL params. */
+	bool UseLegacyProtocol;
 	/**
 	 * Analytics event entry to be cached
 	 */
@@ -117,6 +119,7 @@ TSharedPtr<IAnalyticsProvider> FAnalyticsET::CreateAnalyticsProvider(const FAnal
 		ConfigValues.APIKeyET = GetConfigValue.Execute(Config::GetKeyNameForAPIKey(), true);
 		ConfigValues.APIServerET = GetConfigValue.Execute(Config::GetKeyNameForAPIServer(), false);
 		ConfigValues.AppVersionET = GetConfigValue.Execute(Config::GetKeyNameForAppVersion(), false);
+		ConfigValues.UseLegacyProtocol = FCString::ToBool(*GetConfigValue.Execute(Config::GetKeyNameForUseLegacyProtocol(), false));
 		return CreateAnalyticsProvider(ConfigValues);
 	}
 	else
@@ -147,6 +150,7 @@ FAnalyticsProviderET::FAnalyticsProviderET(const FAnalyticsET::Config& ConfigVal
 	, bShouldCacheEvents(!FParse::Param(FCommandLine::Get(), TEXT("ANALYTICSDISABLECACHING")))
 	, FlushEventsCountdown(MaxCachedElapsedTime)
 	, bInDestructor(false)
+	, UseLegacyProtocol(ConfigValues.UseLegacyProtocol)
 {
 	// if we are not caching events, we are operating in debug mode. Turn on super-verbose analytics logging
 	if (!bShouldCacheEvents)
@@ -282,61 +286,109 @@ void FAnalyticsProviderET::FlushEvents()
 
 		FDateTime CurrentTime = FDateTime::UtcNow();
 
-		TSharedRef< TJsonWriter<TCHAR, TCondensedJsonPrintPolicy<TCHAR> > > JsonWriter = TJsonWriterFactory<TCHAR, TCondensedJsonPrintPolicy<TCHAR> >::Create(&Payload);
-		JsonWriter->WriteObjectStart();
-		JsonWriter->WriteArrayStart(TEXT("Events"));
-		for (int32 EventIdx = 0; EventIdx < CachedEvents.Num(); EventIdx++)
+		if (!UseLegacyProtocol)
 		{
-			const FAnalyticsEventEntry& Entry = CachedEvents[EventIdx];
-			// event entry
+			TSharedRef< TJsonWriter<TCHAR, TCondensedJsonPrintPolicy<TCHAR> > > JsonWriter = TJsonWriterFactory<TCHAR, TCondensedJsonPrintPolicy<TCHAR> >::Create(&Payload);
 			JsonWriter->WriteObjectStart();
-			JsonWriter->WriteValue(TEXT("EventName"), Entry.EventName);
-			FString DateOffset = (CurrentTime - Entry.TimeStamp).ToString();
-			JsonWriter->WriteValue(TEXT("DateOffset"), DateOffset);
-			JsonWriter->WriteValue(TEXT("IsEditor"), FString::FromInt(GIsEditor));
-			if (Entry.Attributes.Num() > 0)
+			JsonWriter->WriteArrayStart(TEXT("Events"));
+			for (int32 EventIdx = 0; EventIdx < CachedEvents.Num(); EventIdx++)
 			{
-				// optional attributes for this event
-				for (int32 AttrIdx = 0; AttrIdx < Entry.Attributes.Num(); AttrIdx++)
+				const FAnalyticsEventEntry& Entry = CachedEvents[EventIdx];
+				// event entry
+				JsonWriter->WriteObjectStart();
+				JsonWriter->WriteValue(TEXT("EventName"), Entry.EventName);
+				FString DateOffset = (CurrentTime - Entry.TimeStamp).ToString();
+				JsonWriter->WriteValue(TEXT("DateOffset"), DateOffset);
+				JsonWriter->WriteValue(TEXT("IsEditor"), FString::FromInt(GIsEditor));
+				if (Entry.Attributes.Num() > 0)
 				{
-					const FAnalyticsEventAttribute& Attr = Entry.Attributes[AttrIdx];				
-					JsonWriter->WriteValue(Attr.AttrName, Attr.AttrValue);
+					// optional attributes for this event
+					for (int32 AttrIdx = 0; AttrIdx < Entry.Attributes.Num(); AttrIdx++)
+					{
+						const FAnalyticsEventAttribute& Attr = Entry.Attributes[AttrIdx];
+						JsonWriter->WriteValue(Attr.AttrName, Attr.AttrValue);
+					}
 				}
+				JsonWriter->WriteObjectEnd();
 			}
+			JsonWriter->WriteArrayEnd();
 			JsonWriter->WriteObjectEnd();
+			JsonWriter->Close();
+
+			FString URLPath = FString::Printf(TEXT("CollectData.1?SessionID=%s&AppID=%s&AppVersion=%s&UserID=%s"),
+				*FPlatformHttp::UrlEncode(SessionID),
+				*FPlatformHttp::UrlEncode(APIKey),
+				*FPlatformHttp::UrlEncode(AppVersion),
+				*FPlatformHttp::UrlEncode(UserID));
+
+			// Recreate the URLPath for logging because we do not want to escape the parameters when logging.
+			// We cannot simply UrlEncode the entire Path after logging it because UrlEncode(Params) != UrlEncode(Param1) & UrlEncode(Param2) ...
+			UE_LOG(LogAnalytics, VeryVerbose, TEXT("[%s] AnalyticsET URL:CollectData.1?SessionID=%s&AppID=%s&AppVersion=%s&UserID=%s. Payload:%s"),
+				*APIKey,
+				*SessionID,
+				*APIKey,
+				*AppVersion,
+				*UserID,
+				*Payload);
+
+			// Create/send Http request for an event
+			TSharedRef<IHttpRequest> HttpRequest = FHttpModule::Get().CreateRequest();
+			HttpRequest->SetHeader(TEXT("Content-Type"), TEXT("application/json; charset=utf-8"));
+
+			HttpRequest->SetURL(APIServer + URLPath);
+			HttpRequest->SetVerb(TEXT("POST"));
+			HttpRequest->SetContentAsString(Payload);
+			if (!bInDestructor)
+			{
+				HttpRequest->OnProcessRequestComplete().BindSP(this, &FAnalyticsProviderET::EventRequestComplete);
+			}
+			HttpRequest->ProcessRequest();
 		}
-		JsonWriter->WriteArrayEnd();
-		JsonWriter->WriteObjectEnd();
-		JsonWriter->Close();
-
-		FString URLPath = FString::Printf(TEXT("CollectData.1?SessionID=%s&AppID=%s&AppVersion=%s&UserID=%s"),
-			*FPlatformHttp::UrlEncode(SessionID),
-			*FPlatformHttp::UrlEncode(APIKey),
-			*FPlatformHttp::UrlEncode(AppVersion),
-			*FPlatformHttp::UrlEncode(UserID));
-
-		// Recreate the URLPath for logging because we do not want to escape the parameters when logging.
-		// We cannot simply UrlEncode the entire Path after logging it because UrlEncode(Params) != UrlEncode(Param1) & UrlEncode(Param2) ...
-		UE_LOG(LogAnalytics, VeryVerbose, TEXT("[%s] AnalyticsET URL:CollectData.1?SessionID=%s&AppID=%s&AppVersion=%s&UserID=%s. Payload:%s"), 
-			*APIKey, 
-			*SessionID,
-			*APIKey,
-			*AppVersion,
-			*UserID,
-			*Payload);
-
-		// Create/send Http request for an event
-		TSharedRef<IHttpRequest> HttpRequest = FHttpModule::Get().CreateRequest();
-		HttpRequest->SetHeader(TEXT("Content-Type"), TEXT("application/json; charset=utf-8"));
-
-		HttpRequest->SetURL(APIServer + URLPath);
-		HttpRequest->SetVerb(TEXT("POST"));
-		HttpRequest->SetContentAsString(Payload);
-		if (!bInDestructor)
+		else
 		{
-			HttpRequest->OnProcessRequestComplete().BindSP(this, &FAnalyticsProviderET::EventRequestComplete);
+			// this is a legacy pathway that doesn't accept batch payloads of cached data. We'll just send one request for each event, which will be slow for a large batch of requests at once.
+			for (const auto& Event : CachedEvents)
+			{
+				FString EventParams;
+				if (Event.Attributes.Num() > 0)
+				{
+					for (int Ndx = 0; Ndx<FMath::Min(Event.Attributes.Num(), 10); ++Ndx)
+					{
+						EventParams += FString::Printf(TEXT("&AttributeName%d=%s&AttributeValue%d=%s"), 
+							Ndx, 
+							*FPlatformHttp::UrlEncode(Event.Attributes[Ndx].AttrName), 
+							Ndx, 
+							*FPlatformHttp::UrlEncode(Event.Attributes[Ndx].AttrValue));
+					}
+				}
+
+				// log out the un-encoded values to make reading the log easier.
+				UE_LOG(LogAnalytics, VeryVerbose, TEXT("[%s] AnalyticsET URL:SendEvent.1?SessionID=%s&AppID=%s&AppVersion=%s&UserID=%s&EventName=%s%s"), 
+					*APIKey,
+					*SessionID,
+					*APIKey,
+					*AppVersion,
+					*UserID,
+					*Event.EventName,
+					*EventParams);
+
+				// Create/send Http request for an event
+				TSharedRef<IHttpRequest> HttpRequest = FHttpModule::Get().CreateRequest();
+				HttpRequest->SetHeader(TEXT("Content-Type"), TEXT("text/plain"));
+				// Don't need to URL encode the APIServer or the EventParams, which are already encoded, and contain parameter separaters that we DON'T want encoded.
+				HttpRequest->SetURL(FString::Printf(TEXT("%sSendEvent.1?SessionID=%s&AppID=%s&AppVersion=%s&UserID=%s&EventName=%s%s"),
+					*APIServer, 
+					*FPlatformHttp::UrlEncode(SessionID), 
+					*FPlatformHttp::UrlEncode(APIKey), 
+					*FPlatformHttp::UrlEncode(AppVersion), 
+					*FPlatformHttp::UrlEncode(UserID), 
+					*FPlatformHttp::UrlEncode(Event.EventName), 
+					*EventParams));
+				HttpRequest->SetVerb(TEXT("GET"));
+				HttpRequest->OnProcessRequestComplete().BindRaw(this, &FAnalyticsProviderET::EventRequestComplete);
+				HttpRequest->ProcessRequest();
+			}
 		}
- 		HttpRequest->ProcessRequest();
 
 		FlushEventsCountdown = MaxCachedElapsedTime;
 		CachedEvents.Empty();
diff --git a/Engine/Source/Runtime/Analytics/AnalyticsET/Public/AnalyticsET.h b/Engine/Source/Runtime/Analytics/AnalyticsET/Public/AnalyticsET.h
--- a/Engine/Source/Runtime/Analytics/AnalyticsET/Public/AnalyticsET.h
+++ b/Engine/Source/Runtime/Analytics/AnalyticsET/Public/AnalyticsET.h
@@ -46,6 +46,11 @@ public:
 		 * If you supply your own Version string, occurrences of "%VERSION%" are replaced with FEngineVersion::Current(). ie, -AnalyticsAppVersion=MyCustomID-%VERSION%.
 		 */
 		FString AppVersionET;
+		/** When true, sends events using the legacy ET protocol that passes all attributes as URL parameters. Defaults to false. */
+		bool UseLegacyProtocol;
+
+		/** Default ctor to ensure all values have their proper default. */
+		Config() : UseLegacyProtocol(false) {}
 
 		/** KeyName required for APIKey configuration. */
 		static FString GetKeyNameForAPIKey() { return TEXT("APIKeyET"); }
@@ -53,6 +58,8 @@ public:
 		static FString GetKeyNameForAPIServer() { return TEXT("APIServerET"); }
 		/** KeyName required for AppVersion configuration. */
 		static FString GetKeyNameForAppVersion() { return TEXT("AppVersionET"); }
+		/** Optional parameter to use the legacy backend protocol. */
+		static FString GetKeyNameForUseLegacyProtocol() { return TEXT("UseLegacyProtocol"); }
 		/** Default value if no APIServer configuration is provided. */
 		static FString GetDefaultAPIServer() { return TEXT("http://devonline-02:/ETAP/"); }
 	};
