Fix for CIS warning in UE4_IOS_Mono_Precompiled_OnMac: all code in AnalyticsSwrve is inside #if PLATFORM_DESKTOP guards, so the static library for it on iOS is empty and the XCode toolchain warns about it. Some other warnings remain due to deprecation of command line arguments in OSX 10.10; Pete Sauerbrei will fix.

diff --git a/Engine/Source/Runtime/Analytics/AnalyticsSwrve/AnalyticsSwrve.Build.cs b/Engine/Source/Runtime/Analytics/AnalyticsSwrve/AnalyticsSwrve.Build.cs
--- a/Engine/Source/Runtime/Analytics/AnalyticsSwrve/AnalyticsSwrve.Build.cs
+++ b/Engine/Source/Runtime/Analytics/AnalyticsSwrve/AnalyticsSwrve.Build.cs
@@ -16,6 +16,11 @@ namespace UnrealBuildTool.Rules
 					"Json",
 				}
 			);
+
+			if(!UnrealBuildTool.IsDesktopPlatform(Target.Platform))
+			{
+				PrecompileForTargets = PrecompileTargetsType.None;
+			}
 		}
 	}
 }
