Copying //UE4/Dev-Build to //UE4/Main

diff --git a/Engine/Source/Programs/AutomationTool/Program.cs b/Engine/Source/Programs/AutomationTool/Program.cs
--- a/Engine/Source/Programs/AutomationTool/Program.cs
+++ b/Engine/Source/Programs/AutomationTool/Program.cs
@@ -29,7 +29,7 @@ namespace AutomationTool
 				AssemblyUtils.InstallAssemblyResolver(Path.GetDirectoryName(Assembly.GetEntryAssembly().GetOriginalLocation()));
 				HostPlatform.Initialize();
 
-				Log.TraceInformation("{2}: Running on {0} as a {1}-bit process.", HostPlatform.Current.GetType().Name, Environment.Is64BitProcess ? 64 : 32, DateTime.UtcNow.ToString("o"));
+				Log.TraceVerbose("{2}: Running on {0} as a {1}-bit process.", HostPlatform.Current.GetType().Name, Environment.Is64BitProcess ? 64 : 32, DateTime.UtcNow.ToString("o"));
 
 				XmlConfigLoader.Init();
 
diff --git a/Engine/Source/Programs/AutomationTool/Scripts/BuildProjectCommand.Automation.cs b/Engine/Source/Programs/AutomationTool/Scripts/BuildProjectCommand.Automation.cs
--- a/Engine/Source/Programs/AutomationTool/Scripts/BuildProjectCommand.Automation.cs
+++ b/Engine/Source/Programs/AutomationTool/Scripts/BuildProjectCommand.Automation.cs
@@ -53,7 +53,7 @@ public partial class Project : CommandUtils
 		var CrashReportPlatforms = new HashSet<UnrealTargetPlatform>();
 
 		// Setup editor targets
-		if (Params.HasEditorTargets && !Params.Rocket && (TargetMask & ProjectBuildTargets.Editor) == ProjectBuildTargets.Editor)
+		if (Params.HasEditorTargets && !Automation.RunningRocket() && (TargetMask & ProjectBuildTargets.Editor) == ProjectBuildTargets.Editor)
 		{
 			// @todo Mac: proper platform detection
 			UnrealTargetPlatform EditorPlatform = HostPlatform.Current.HostEditorPlatform;
@@ -114,7 +114,7 @@ public partial class Project : CommandUtils
 				}
 			}
 		}
-		if (!Params.NoBootstrapExe && !Params.Rocket && (TargetMask & ProjectBuildTargets.Bootstrap) == ProjectBuildTargets.Bootstrap)
+		if (!Params.NoBootstrapExe && !Automation.RunningRocket() && (TargetMask & ProjectBuildTargets.Bootstrap) == ProjectBuildTargets.Bootstrap)
 		{
 			UnrealBuildTool.UnrealTargetPlatform[] BootstrapPackagedGamePlatforms = { UnrealBuildTool.UnrealTargetPlatform.Win32, UnrealBuildTool.UnrealTargetPlatform.Win64 };
 			foreach(UnrealBuildTool.UnrealTargetPlatform BootstrapPackagedGamePlatform in BootstrapPackagedGamePlatforms)
@@ -125,7 +125,7 @@ public partial class Project : CommandUtils
 				}
 			}
 		}
-		if (Params.CrashReporter && !Params.Rocket && (TargetMask & ProjectBuildTargets.CrashReporter) == ProjectBuildTargets.CrashReporter)
+		if (Params.CrashReporter && !Automation.RunningRocket() && (TargetMask & ProjectBuildTargets.CrashReporter) == ProjectBuildTargets.CrashReporter)
 		{
 			foreach (var CrashReportPlatform in CrashReportPlatforms)
 			{
@@ -135,7 +135,7 @@ public partial class Project : CommandUtils
 				}
 			}
 		}
-		if (Params.HasProgramTargets && !Params.Rocket && (TargetMask & ProjectBuildTargets.Programs) == ProjectBuildTargets.Programs)
+		if (Params.HasProgramTargets && !Automation.RunningRocket() && (TargetMask & ProjectBuildTargets.Programs) == ProjectBuildTargets.Programs)
 		{
 			foreach (var BuildConfig in Params.ClientConfigsToBuild)
 			{
diff --git a/Engine/Source/Programs/AutomationTool/Scripts/CopyBuildToStagingDirectory.Automation.cs b/Engine/Source/Programs/AutomationTool/Scripts/CopyBuildToStagingDirectory.Automation.cs
--- a/Engine/Source/Programs/AutomationTool/Scripts/CopyBuildToStagingDirectory.Automation.cs
+++ b/Engine/Source/Programs/AutomationTool/Scripts/CopyBuildToStagingDirectory.Automation.cs
@@ -1709,7 +1709,7 @@ public partial class Project : CommandUtils
 				ExecutablesToStage,
 				InDedicatedServer,
 				Params.Cook || Params.CookOnTheFly,
-				Params.CrashReporter && !(StagePlatform == UnrealTargetPlatform.Linux && Params.Rocket), // can't include the crash reporter from binary Linux builds
+				Params.CrashReporter && !(StagePlatform == UnrealTargetPlatform.Linux && Automation.RunningRocket()), // can't include the crash reporter from binary Linux builds
 				Params.Stage,
 				Params.CookOnTheFly,
 				Params.Archive,
diff --git a/Engine/Source/Programs/AutomationTool/Scripts/RocketBuild.Automation.cs b/Engine/Source/Programs/AutomationTool/Scripts/RocketBuild.Automation.cs
--- a/Engine/Source/Programs/AutomationTool/Scripts/RocketBuild.Automation.cs
+++ b/Engine/Source/Programs/AutomationTool/Scripts/RocketBuild.Automation.cs
@@ -289,13 +289,13 @@ namespace Rocket
         {
 			foreach(UnrealTargetPlatform HostPlatform in BranchConfig.HostPlatforms)
 			{
-				AddDependency(FilterRocketNode.StaticGetFullName(HostPlatform));
-				AddDependency(BuildDerivedDataCacheNode.StaticGetFullName(HostPlatform));
+				AddPseudodependency(FilterRocketNode.StaticGetFullName(HostPlatform));
+				AddPseudodependency(BuildDerivedDataCacheNode.StaticGetFullName(HostPlatform));
 
 				SingleTargetProperties BuildPatchTool = BranchConfig.Branch.FindProgram("BuildPatchTool");
 				if(BuildPatchTool.Rules != null)
 				{
-					AddDependency(GUBP.SingleInternalToolsNode.StaticGetFullName(HostPlatform, BuildPatchTool));
+					AddPseudodependency(GUBP.SingleInternalToolsNode.StaticGetFullName(HostPlatform, BuildPatchTool));
 				}
 			}
 		}
@@ -444,6 +444,14 @@ namespace Rocket
 		{
 			AddNodeToStrip(GUBP.ToolsForCompileNode.StaticGetFullName(HostPlatform));
 			AddNodeToStrip(GUBP.ToolsNode.StaticGetFullName(HostPlatform));
+
+			SingleTargetProperties BuildPatchTool = BranchConfig.Branch.FindProgram("BuildPatchTool");
+			if (BuildPatchTool.Rules == null)
+			{
+				throw new AutomationException("Could not find program BuildPatchTool.");
+			}
+			AddNodeToStrip(GUBP.SingleInternalToolsNode.StaticGetFullName(HostPlatform, BuildPatchTool));
+
 			AgentSharingGroup = "ToolsGroup" + StaticGetHostPlatformSuffix(InHostPlatform);
 		}
 
@@ -617,6 +625,14 @@ namespace Rocket
 		{
 			AddNodeToSign(GUBP.ToolsForCompileNode.StaticGetFullName(HostPlatform));
 			AddNodeToSign(GUBP.ToolsNode.StaticGetFullName(HostPlatform));
+
+			SingleTargetProperties BuildPatchTool = BranchConfig.Branch.FindProgram("BuildPatchTool");
+			if (BuildPatchTool.Rules == null)
+			{
+				throw new AutomationException("Could not find program BuildPatchTool.");
+			}
+			AddNodeToSign(GUBP.SingleInternalToolsNode.StaticGetFullName(HostPlatform, BuildPatchTool));
+
 			AddDependency(StripRocketToolsNode.StaticGetFullName(HostPlatform));
 			AgentSharingGroup = "ToolsGroup" + StaticGetHostPlatformSuffix(InHostPlatform);
 		}
@@ -728,6 +744,13 @@ namespace Rocket
 			AddDependency(GUBP.RootEditorNode.StaticGetFullName(HostPlatform));
 			AddDependency(GUBP.ToolsNode.StaticGetFullName(HostPlatform));
 
+			SingleTargetProperties BuildPatchTool = BranchConfig.Branch.FindProgram("BuildPatchTool");
+			if (BuildPatchTool.Rules == null)
+			{
+				throw new AutomationException("Could not find program BuildPatchTool.");
+			}
+			AddDependency(GUBP.SingleInternalToolsNode.StaticGetFullName(HostPlatform, BuildPatchTool));
+
 			// Add all the monolithic builds from their appropriate source host platform
 			foreach(UnrealTargetPlatform TargetPlatform in TargetPlatforms)
 			{
@@ -836,6 +859,9 @@ namespace Rocket
 			AddRuleForBuildProducts(Filter, BranchConfig, GUBP.RootEditorNode.StaticGetFullName(HostPlatform), FileFilterType.Include);
 			AddRuleForBuildProducts(Filter, BranchConfig, GUBP.ToolsNode.StaticGetFullName(HostPlatform), FileFilterType.Include);
 
+			SingleTargetProperties BuildPatchTool = BranchConfig.Branch.FindProgram("BuildPatchTool");
+			AddRuleForBuildProducts(Filter, BranchConfig, GUBP.SingleInternalToolsNode.StaticGetFullName(HostPlatform, BuildPatchTool), FileFilterType.Include);
+
 			// Include win64 tools on Mac, to get the win64 build of UBT, UAT and IPP
 			if (HostPlatform == UnrealTargetPlatform.Mac && BranchConfig.HostPlatforms.Contains(UnrealTargetPlatform.Win64))
 			{
@@ -1282,7 +1308,7 @@ namespace Rocket
 			AddDependency(GUBP.ToolsForCompileNode.StaticGetFullName(HostPlatform));
 			AddDependency(GUBP.RootEditorNode.StaticGetFullName(HostPlatform));
 			AddDependency(GUBP.ToolsNode.StaticGetFullName(HostPlatform));
-			AddDependency(WaitToMakeRocketBuild.StaticGetFullName());
+			AddPseudodependency(WaitToMakeRocketBuild.StaticGetFullName());
 
 			foreach(UnrealTargetPlatform TargetPlatform in TargetPlatforms)
 			{
diff --git a/Engine/Source/Programs/CrashReporter/CrashReportClient/Private/CrashReportAnalytics.cpp b/Engine/Source/Programs/CrashReporter/CrashReportClient/Private/CrashReportAnalytics.cpp
--- a/Engine/Source/Programs/CrashReporter/CrashReportClient/Private/CrashReportAnalytics.cpp
+++ b/Engine/Source/Programs/CrashReporter/CrashReportClient/Private/CrashReportAnalytics.cpp
@@ -4,8 +4,7 @@
 
 #include "CrashReportAnalytics.h"
 #include "EngineBuildSettings.h"
-#include "Runtime/Analytics/Analytics/Public/Analytics.h"
-#include "Runtime/Analytics/Analytics/Public/Interfaces/IAnalyticsProvider.h"
+#include "Runtime/Analytics/AnalyticsET/Public/AnalyticsET.h"
 
 // WARNING! Copied from EngineAnalytics.cpp
 
@@ -17,10 +16,10 @@ TSharedPtr<IAnalyticsProvider> FCrashReportAnalytics::Analytics;
  * External code should bind this delegate if engine analytics are desired,
  * preferably in private code that won't be redistributed.
  */
-FAnalytics::FProviderConfigurationDelegate& GetCrashReportAnalyticsOverrideConfigDelegate()
+FAnalyticsET::Config& GetCrashReportAnalyticsOverrideConfig()
 {
-	static FAnalytics::FProviderConfigurationDelegate Delegate;
-	return Delegate;
+	static FAnalyticsET::Config Config;
+	return Config;
 }
 
 
@@ -37,46 +36,10 @@ void FCrashReportAnalytics::Initialize()
 {
 	checkf(!bIsInitialized, TEXT("FCrashReportAnalytics::Initialize called more than once."));
 
+	if (!GetCrashReportAnalyticsOverrideConfig().APIServerET.IsEmpty())
 	{
-		// Setup some default engine analytics if there is nothing custom bound
-		FAnalytics::FProviderConfigurationDelegate DefaultEngineAnalyticsConfig;
-		DefaultEngineAnalyticsConfig.BindStatic(
-			[]( const FString& KeyName, bool bIsValueRequired ) -> FString
-		{
-			static TMap<FString, FString> ConfigMap;
-			if( ConfigMap.Num() == 0 )
-			{
-				ConfigMap.Add( TEXT( "ProviderModuleName" ), TEXT( "AnalyticsET" ) );
-				ConfigMap.Add( TEXT( "APIServerET" ), TEXT( "http://etsource.epicgames.com/ET2/" ) );
-
-				// We always use the "Release" analytics account unless we're running in analytics test mode (usually with
-				// a command-line parameter), or we're an internal Epic build
-				bool bUseReleaseAccount =
-					(FAnalytics::Get().GetBuildType() == FAnalytics::Development ||
-					FAnalytics::Get().GetBuildType() == FAnalytics::Release) &&
-					!FEngineBuildSettings::IsInternalBuild();	// Internal Epic build
-
-				ConfigMap.Add( TEXT( "APIKeyET" ), bUseReleaseAccount ? TEXT("CrashReporter.Release") : TEXT("CrashReporter.Dev") );
-			}
-
-			// Check for overrides
-			if( GetCrashReportAnalyticsOverrideConfigDelegate().IsBound() )
-			{
-				const FString OverrideValue = GetCrashReportAnalyticsOverrideConfigDelegate().Execute( KeyName, bIsValueRequired );
-				if( !OverrideValue.IsEmpty() )
-				{
-					return OverrideValue;
-				}
-			}
-
-			FString* ConfigValue = ConfigMap.Find( KeyName );
-			return ConfigValue != NULL ? *ConfigValue : TEXT( "" );
-		} );
-
 		// Connect the engine analytics provider (if there is a configuration delegate installed)
-		Analytics = FAnalytics::Get().CreateAnalyticsProvider(
-			FName( *DefaultEngineAnalyticsConfig.Execute( TEXT( "ProviderModuleName" ), true ) ),
-			DefaultEngineAnalyticsConfig );
+		Analytics = FAnalyticsET::Get().CreateAnalyticsProvider(GetCrashReportAnalyticsOverrideConfig());
 		if( Analytics.IsValid() )
 		{
 			Analytics->SetUserID(FString::Printf(TEXT("%s|%s|%s"), *FPlatformMisc::GetMachineId().ToString(EGuidFormats::Digits).ToLower(), *FPlatformMisc::GetEpicAccountId(), *FPlatformMisc::GetOperatingSystemId()));
diff --git a/Engine/Source/Programs/UnrealBuildTool/Configuration/UEBuildConfiguration.cs b/Engine/Source/Programs/UnrealBuildTool/Configuration/UEBuildConfiguration.cs
--- a/Engine/Source/Programs/UnrealBuildTool/Configuration/UEBuildConfiguration.cs
+++ b/Engine/Source/Programs/UnrealBuildTool/Configuration/UEBuildConfiguration.cs
@@ -359,8 +359,7 @@ namespace UnrealBuildTool
 			bCompileSimplygon = bCompileSimplygon
 				&& Directory.Exists(UEBuildConfiguration.UEThirdPartySourceDirectory + "NotForLicensees") == true
 				&& Directory.Exists(UEBuildConfiguration.UEThirdPartySourceDirectory + "NotForLicensees/Simplygon") == true
-				&& Directory.Exists("Developer/SimplygonMeshReduction") == true
-				&& !(ProjectFileGenerator.bGenerateProjectFiles && ProjectFileGenerator.bGeneratingRocketProjectFiles);
+				&& Directory.Exists("Developer/SimplygonMeshReduction") == true;
 		}
 
 		/// <summary>
diff --git a/Engine/Source/Programs/UnrealBuildTool/Configuration/UEBuildPlatform.cs b/Engine/Source/Programs/UnrealBuildTool/Configuration/UEBuildPlatform.cs
--- a/Engine/Source/Programs/UnrealBuildTool/Configuration/UEBuildPlatform.cs
+++ b/Engine/Source/Programs/UnrealBuildTool/Configuration/UEBuildPlatform.cs
@@ -476,20 +476,6 @@ namespace UnrealBuildTool
 			}
 		}
 
-		// Included for compatibility during //UE4/Main import
-		[Obsolete]
-		public CPPTargetPlatform GetCPPTargetPlatform(UnrealTargetPlatform ForPlatform)
-		{
-			if(Platform == ForPlatform)
-			{
-				return DefaultCppPlatform;
-			}
-			else
-			{
-				throw new BuildException("Target platform {0} can not be compiled on {1}", ForPlatform, Platform);
-			}
-		}
-
 		/// <summary>
 		/// Returns the delimiter used to separate paths in the PATH environment variable for the platform we are executing on.
 		/// </summary>
diff --git a/Engine/Source/Programs/UnrealBuildTool/Configuration/UEBuildTarget.cs b/Engine/Source/Programs/UnrealBuildTool/Configuration/UEBuildTarget.cs
--- a/Engine/Source/Programs/UnrealBuildTool/Configuration/UEBuildTarget.cs
+++ b/Engine/Source/Programs/UnrealBuildTool/Configuration/UEBuildTarget.cs
@@ -269,7 +269,7 @@ namespace UnrealBuildTool
 
 			// Settings for creating/using static libraries for the engine
 			bool bPrecompile = false;
-			bool bUsePrecompiled = UnrealBuildTool.RunningRocket();
+			bool bUsePrecompiled = UnrealBuildTool.IsEngineInstalled();
 
 			// Combine the two arrays of arguments
 			List<string> Arguments = new List<string>(SourceArguments.Length);
@@ -1300,7 +1300,7 @@ namespace UnrealBuildTool
 		/// </summary>
 		private void CleanUnrealHeaderTool()
 		{
-			if (!UnrealBuildTool.RunningRocket())
+			if (!UnrealBuildTool.IsEngineInstalled())
 			{
 				var UBTArguments = new StringBuilder();
 
@@ -1348,8 +1348,8 @@ namespace UnrealBuildTool
 						continue;
 					}
 
-					// if we're cleaning and in an installed or rocket binary and the build product path matches the engine install location, we don't want to delete them
-					if (UEBuildConfiguration.bCleanProject && (UnrealBuildTool.IsEngineInstalled() || UnrealBuildTool.RunningRocket()) && BuildProduct.Path.Contains(Path.GetFullPath(BuildConfiguration.RelativeEnginePath)))
+					// if we're cleaning and in an installed binary and the build product path matches the engine install location, we don't want to delete them
+					if (UEBuildConfiguration.bCleanProject && UnrealBuildTool.IsEngineInstalled() && BuildProduct.Path.Contains(Path.GetFullPath(BuildConfiguration.RelativeEnginePath)))
 					{
 						continue;
 					}
@@ -1393,7 +1393,7 @@ namespace UnrealBuildTool
 				{
 					if (Binary.Config.IntermediateDirectory.Exists())
 					{
-						if (!UnrealBuildTool.RunningRocket() || !Binary.Config.IntermediateDirectory.IsUnderDirectory(UnrealBuildTool.EngineDirectory))
+						if (!UnrealBuildTool.IsEngineInstalled() || !Binary.Config.IntermediateDirectory.IsUnderDirectory(UnrealBuildTool.EngineDirectory))
 						{
 							CleanDirectory(Binary.Config.IntermediateDirectory.FullName);
 						}
@@ -1464,9 +1464,9 @@ namespace UnrealBuildTool
 						CleanDirectory(ProjectIntermediateDirectory.FullName);
 					}
 				}
-				if (!UnrealBuildTool.RunningRocket())
+				if (!UnrealBuildTool.IsEngineInstalled())
 				{
-					// This is always under Rocket installation folder
+					// This is always under Engine installation folder
 					if (GlobalLinkEnvironment.Config.IntermediateDirectory.Exists())
 					{
 						Log.TraceVerbose("\tDeleting " + GlobalLinkEnvironment.Config.IntermediateDirectory);
@@ -1475,7 +1475,7 @@ namespace UnrealBuildTool
 				}
 				else if (ShouldCompileMonolithic())
 				{
-					// Only in monolithic, otherwise it's pointing to Rocket installation folder
+					// Only in monolithic, otherwise it's pointing to Engine installation folder
 					if (GlobalLinkEnvironment.Config.OutputDirectory.Exists())
 					{
 						Log.TraceVerbose("\tDeleting " + GlobalLinkEnvironment.Config.OutputDirectory);
@@ -1606,8 +1606,8 @@ namespace UnrealBuildTool
 						continue;
 					}
 
-					// don't add any installed or rocket binary and the build product path matches the engine install location, we don't want to delete them
-					if ((UnrealBuildTool.IsEngineInstalled() || UnrealBuildTool.RunningRocket()) && BuildProduct.Path.Contains(Path.GetFullPath(BuildConfiguration.RelativeEnginePath)))
+					// don't add any installed binary and the build product path matches the engine install location, we don't want to delete them
+					if (UnrealBuildTool.IsEngineInstalled() && BuildProduct.Path.Contains(Path.GetFullPath(BuildConfiguration.RelativeEnginePath)))
 					{
 						continue;
 					}
@@ -1807,7 +1807,7 @@ namespace UnrealBuildTool
 		public void GenerateManifest()
 		{
 			FileReference ManifestPath;
-			if (UnrealBuildTool.RunningRocket() && ProjectFile != null)
+			if (UnrealBuildTool.IsEngineInstalled() && ProjectFile != null)
 			{
 				ManifestPath = FileReference.Combine(ProjectFile.Directory, BuildConfiguration.BaseIntermediateFolder, "Manifest.xml");
 			}
@@ -2449,7 +2449,7 @@ namespace UnrealBuildTool
 			}
 
 			[REDACTED_EMAIL]: Will users be able to rebuild UnrealHeaderTool? NO
-			if (!ProjectFileGenerator.bGenerateProjectFiles && UnrealBuildTool.RunningRocket() && AppName != "UnrealHeaderTool")
+			if (!ProjectFileGenerator.bGenerateProjectFiles && UnrealBuildTool.IsEngineInstalled() && AppName != "UnrealHeaderTool")
 			{
 				var FilteredBinaries = new List<UEBuildBinary>();
 
@@ -3441,36 +3441,26 @@ namespace UnrealBuildTool
 
 			string OutputAppName = GetAppName();
 
+			UnrealTargetConfiguration EngineTargetConfiguration = Configuration == UnrealTargetConfiguration.DebugGame ? UnrealTargetConfiguration.Development : Configuration;
+			GlobalCompileEnvironment.Config.OutputDirectory = DirectoryReference.Combine(UnrealBuildTool.EngineDirectory, BuildConfiguration.PlatformIntermediateFolder, OutputAppName, EngineTargetConfiguration.ToString());
+
 			// Installed Engine intermediates go to the project's intermediate folder. Installed Engine never writes to the engine intermediate folder. (Those files are immutable)
 			// Also, when compiling in monolithic, all intermediates go to the project's folder.  This is because a project can change definitions that affects all engine translation
 			// units too, so they can't be shared between different targets.  They are effectively project-specific engine intermediates.
 			if (UnrealBuildTool.IsEngineInstalled() || (ProjectFile != null && ShouldCompileMonolithic()))
 			{
-				var IntermediateConfiguration = Configuration;
-				if (UnrealBuildTool.RunningRocket())
-				{
-					// Only Development and Shipping are supported for engine modules
-					IntermediateConfiguration = Configuration == UnrealTargetConfiguration.DebugGame ? UnrealTargetConfiguration.Development : Configuration;
-				}
-
-				GlobalCompileEnvironment.Config.OutputDirectory = DirectoryReference.Combine(UnrealBuildTool.EngineDirectory, BuildConfiguration.PlatformIntermediateFolder, OutputAppName, IntermediateConfiguration.ToString());
 				if (ShouldCompileMonolithic())
 				{
 					if (ProjectFile != null)
 					{
-						GlobalCompileEnvironment.Config.OutputDirectory = DirectoryReference.Combine(ProjectFile.Directory, BuildConfiguration.PlatformIntermediateFolder, OutputAppName, IntermediateConfiguration.ToString());
+						GlobalCompileEnvironment.Config.OutputDirectory = DirectoryReference.Combine(ProjectFile.Directory, BuildConfiguration.PlatformIntermediateFolder, OutputAppName, Configuration.ToString());
 					}
 					else if (ForeignPlugins.Count > 0)
 					{
-						GlobalCompileEnvironment.Config.OutputDirectory = DirectoryReference.Combine(ForeignPlugins[0].Directory, BuildConfiguration.PlatformIntermediateFolder, OutputAppName, IntermediateConfiguration.ToString());
+						GlobalCompileEnvironment.Config.OutputDirectory = DirectoryReference.Combine(ForeignPlugins[0].Directory, BuildConfiguration.PlatformIntermediateFolder, OutputAppName, Configuration.ToString());
 					}
 				}
 			}
-			else
-			{
-				var GlobalEnvironmentConfiguration = Configuration == UnrealTargetConfiguration.DebugGame ? UnrealTargetConfiguration.Development : Configuration;
-				GlobalCompileEnvironment.Config.OutputDirectory = DirectoryReference.Combine(UnrealBuildTool.EngineDirectory, BuildConfiguration.PlatformIntermediateFolder, OutputAppName, GlobalEnvironmentConfiguration.ToString());
-			}
 
 			// By default, shadow source files for this target in the root OutputDirectory
 			GlobalCompileEnvironment.Config.LocalShadowDirectory = GlobalCompileEnvironment.Config.OutputDirectory;
diff --git a/Engine/Source/Programs/UnrealBuildTool/Linux/UEBuildLinux.cs b/Engine/Source/Programs/UnrealBuildTool/Linux/UEBuildLinux.cs
--- a/Engine/Source/Programs/UnrealBuildTool/Linux/UEBuildLinux.cs
+++ b/Engine/Source/Programs/UnrealBuildTool/Linux/UEBuildLinux.cs
@@ -106,7 +106,7 @@ namespace UnrealBuildTool
 		/// </summary>
 		public override void ValidateUEBuildConfiguration()
 		{
-			if (ProjectFileGenerator.bGenerateProjectFiles && !ProjectFileGenerator.bGeneratingRocketProjectFiles)
+			if (ProjectFileGenerator.bGenerateProjectFiles && !UnrealBuildTool.RunningRocket())
 			{
 				// When generating non-Rocket project files we need intellisense generator to include info from all modules,
 				// including editor-only third party libs
diff --git a/Engine/Source/Programs/UnrealBuildTool/Mac/MacToolChain.cs b/Engine/Source/Programs/UnrealBuildTool/Mac/MacToolChain.cs
--- a/Engine/Source/Programs/UnrealBuildTool/Mac/MacToolChain.cs
+++ b/Engine/Source/Programs/UnrealBuildTool/Mac/MacToolChain.cs
@@ -871,7 +871,7 @@ namespace UnrealBuildTool
 
 			if (!bIsBuildingLibrary)
 			{
-				if (UnrealBuildTool.RunningRocket() || (Utils.IsRunningOnMono && LinkEnvironment.Config.bIsCrossReferenced == false))
+				if (UnrealBuildTool.IsEngineInstalled() || (Utils.IsRunningOnMono && LinkEnvironment.Config.bIsCrossReferenced == false))
 				{
 					foreach (string Library in EngineAndGameLibraries)
 					{
diff --git a/Engine/Source/Programs/UnrealBuildTool/Mac/UEBuildMac.cs b/Engine/Source/Programs/UnrealBuildTool/Mac/UEBuildMac.cs
--- a/Engine/Source/Programs/UnrealBuildTool/Mac/UEBuildMac.cs
+++ b/Engine/Source/Programs/UnrealBuildTool/Mac/UEBuildMac.cs
@@ -81,7 +81,7 @@ namespace UnrealBuildTool
 
 		public override void ValidateUEBuildConfiguration()
 		{
-			if (ProjectFileGenerator.bGenerateProjectFiles && !ProjectFileGenerator.bGeneratingRocketProjectFiles)
+			if (ProjectFileGenerator.bGenerateProjectFiles && !UnrealBuildTool.RunningRocket())
 			{
 				// When generating non-Rocket project files we need intellisense generator to include info from all modules, including editor-only third party libs
 				UEBuildConfiguration.bCompileLeanAndMeanUE = false;
diff --git a/Engine/Source/Programs/UnrealBuildTool/System/ActionHistory.cs b/Engine/Source/Programs/UnrealBuildTool/System/ActionHistory.cs
--- a/Engine/Source/Programs/UnrealBuildTool/System/ActionHistory.cs
+++ b/Engine/Source/Programs/UnrealBuildTool/System/ActionHistory.cs
@@ -118,8 +118,8 @@ namespace UnrealBuildTool
 			}
 			else
 			{
-				// Shared action history (unless this is a rocket target)
-				Folder = (UnrealBuildTool.RunningRocket() && Target.ProjectFile != null) ?
+				// Shared action history (unless this is an installed build target)
+				Folder = (UnrealBuildTool.IsEngineInstalled() && Target.ProjectFile != null) ?
 					DirectoryReference.Combine(Target.ProjectFile.Directory, BuildConfiguration.BaseIntermediateFolder) :
 					DirectoryReference.Combine(UnrealBuildTool.EngineDirectory, BuildConfiguration.BaseIntermediateFolder);
 			}
diff --git a/Engine/Source/Programs/UnrealBuildTool/System/ExternalExecution.cs b/Engine/Source/Programs/UnrealBuildTool/System/ExternalExecution.cs
--- a/Engine/Source/Programs/UnrealBuildTool/System/ExternalExecution.cs
+++ b/Engine/Source/Programs/UnrealBuildTool/System/ExternalExecution.cs
@@ -417,9 +417,9 @@ namespace UnrealBuildTool
 		private static DateTime GetCoreGeneratedTimestamp(string ModuleName, string ModuleGeneratedCodeDirectory)
 		{
 			DateTime Timestamp;
-			if (UnrealBuildTool.RunningRocket())
+			if (UnrealBuildTool.IsEngineInstalled())
 			{
-				// In Rocket, we don't check the timestamps on engine headers.  Default to a very old date.
+				// In Installed Builds, we don't check the timestamps on engine headers.  Default to a very old date.
 				Timestamp = DateTime.MinValue;
 			}
 			else
@@ -705,8 +705,8 @@ namespace UnrealBuildTool
 
 				if (!bIsBuildingUHT && bUHTNeedsToRun)
 				{
-					// Always build UnrealHeaderTool if header regeneration is required, unless we're running within a Rocket ecosystem or hot-reloading
-					if (UnrealBuildTool.RunningRocket() == false &&
+					// Always build UnrealHeaderTool if header regeneration is required, unless we're running within an installed ecosystem or hot-reloading
+					if (UnrealBuildTool.IsEngineInstalled() == false &&
 						UEBuildConfiguration.bDoNotBuildUHT == false &&
 						UEBuildConfiguration.bHotReloadFromIDE == false &&
 						!(bHaveHeaderTool && !UnrealBuildTool.IsGatheringBuild && UnrealBuildTool.IsAssemblingBuild))	// If running in "assembler only" mode, we assume UHT is already up to date for much faster iteration!
diff --git a/Engine/Source/Programs/UnrealBuildTool/System/ProjectFileGenerator.cs b/Engine/Source/Programs/UnrealBuildTool/System/ProjectFileGenerator.cs
--- a/Engine/Source/Programs/UnrealBuildTool/System/ProjectFileGenerator.cs
+++ b/Engine/Source/Programs/UnrealBuildTool/System/ProjectFileGenerator.cs
@@ -126,9 +126,6 @@ namespace UnrealBuildTool
 		/// True if we're generating lightweight project files for a single game only, excluding most engine code, documentation, etc.
 		public bool bGeneratingGameProjectFiles = false;
 
-		/// True if we're generating project files to be used with Rocket
-		public static bool bGeneratingRocketProjectFiles = false;
-
 		/// Optional list of platforms to generate projects for
 		readonly List<UnrealTargetPlatform> ProjectPlatforms = new List<UnrealTargetPlatform>();
 
@@ -334,7 +331,7 @@ namespace UnrealBuildTool
 				}
 				IntermediateProjectFilesPath = DirectoryReference.Combine(MasterProjectPath, "Intermediate", "ProjectFiles");
 			}
-			else if( bGeneratingRocketProjectFiles )
+			else if( UnrealBuildTool.RunningRocket() )
 			{
 				Log.TraceInformation("Discovering modules, targets and source code for project...");
 
@@ -389,16 +386,6 @@ namespace UnrealBuildTool
 			// Build the list of games to generate projects for
 			List<UProjectInfo> AllGameProjects = FindGameProjects();
 
-			var AssemblyName = "ProjectFileGenerator";
-			if (bGeneratingRocketProjectFiles)
-			{
-				AssemblyName = "RocketProjectFileGenerator";
-			}
-			else if (bGeneratingGameProjectFiles)
-			{
-				AssemblyName = GameProjectName + "ProjectFileGenerator";
-			}
-
 			// Find all of the module files.  This will filter out any modules or targets that don't belong to platforms
 			// we're generating project files for.
 			var AllModuleFiles = DiscoverModules(AllGameProjects);
@@ -793,10 +780,7 @@ namespace UnrealBuildTool
 				}
 			}
 
-			bGeneratingRocketProjectFiles = UnrealBuildTool.RunningRocket();
-
-
-			if( bGeneratingRocketProjectFiles )
+			if( UnrealBuildTool.RunningRocket() )
 			{
 				// Make sure we can get a valid game name out of this project
 				var GameName = OnlyGameProject.GetFileNameWithoutExtension();
@@ -805,6 +789,8 @@ namespace UnrealBuildTool
 					throw new BuildException("A valid Rocket game project was not found in the specified location (" + OnlyGameProject.Directory.FullName + ")");
 				}
 
+				bIncludeTestAndShippingConfigs = true;
+
 				IncludeEngineSource = true;
 				IncludeDocumentation = false;
 				IncludeBuildSystemFiles = false;
@@ -1516,16 +1502,10 @@ namespace UnrealBuildTool
 				bool IsEngineModule = CurModuleFile.IsUnderDirectory(UnrealBuildTool.EngineDirectory);
 				bool IsThirdPartyModule = CurModuleFile.IsUnderDirectory(EngineSourceThirdPartyDirectory);
 
-				if( bGeneratingGameProjectFiles || bGeneratingRocketProjectFiles )
+				if( IsEngineModule && !IncludeEngineSource )
 				{
-					if( IsEngineModule )
-					{
-						if( !IncludeEngineSource )
-						{
-							// We were asked to exclude engine modules from the generated projects
-							WantProjectFileForModule = false;
-						}
-					}
+					// We were asked to exclude engine modules from the generated projects
+					WantProjectFileForModule = false;
 				}
 
 				if( WantProjectFileForModule )
@@ -1735,9 +1715,9 @@ namespace UnrealBuildTool
 						{
 							EngineProject = ProjectFile;
 							BaseFolder = UnrealBuildTool.EngineDirectory;
-							if (bGeneratingRocketProjectFiles)
+							if (UnrealBuildTool.IsEngineInstalled())
 							{
-								// Allow engine projects to be created but not built for Rocket
+								// Allow engine projects to be created but not built for Installed Engine builds
 								EngineProject.IsForeignProject = false;
 								EngineProject.IsGeneratedProject = true;
 								EngineProject.IsStubProject = true;
diff --git a/Engine/Source/Programs/UnrealBuildTool/System/RulesCompiler.cs b/Engine/Source/Programs/UnrealBuildTool/System/RulesCompiler.cs
--- a/Engine/Source/Programs/UnrealBuildTool/System/RulesCompiler.cs
+++ b/Engine/Source/Programs/UnrealBuildTool/System/RulesCompiler.cs
@@ -451,7 +451,7 @@ namespace UnrealBuildTool
 		/// <param name="ModuleNames">The names of the modules to add</param>
 		public void AddThirdPartyPrivateStaticDependencies(TargetInfo Target, params string[] InModuleNames)
 		{
-			if (UnrealBuildTool.RunningRocket() == false || Target.Type == TargetRules.TargetType.Game)
+			if (UnrealBuildTool.RunningRocket() == false || Target.bIsMonolithic.Value == true)
 			{
 				PrivateDependencyModuleNames.AddRange(InModuleNames);
 			}
@@ -466,7 +466,7 @@ namespace UnrealBuildTool
 		/// <param name="ModuleNames">The names of the modules to add</param>
 		public void AddThirdPartyPrivateDynamicDependencies(TargetInfo Target, params string[] InModuleNames)
 		{
-			if (UnrealBuildTool.RunningRocket() == false || Target.Type == TargetRules.TargetType.Game)
+			if (UnrealBuildTool.RunningRocket() == false || Target.bIsMonolithic.Value == true)
 			{
 				PrivateIncludePathModuleNames.AddRange(InModuleNames);
 				DynamicallyLoadedModuleNames.AddRange(InModuleNames);
@@ -857,7 +857,7 @@ namespace UnrealBuildTool
 		/// <returns>True if the target should use the shared build environment</returns>
 		public virtual bool ShouldUseSharedBuildEnvironment(TargetInfo Target)
 		{
-			return UnrealBuildTool.RunningRocket() || (Target.Type != TargetType.Program && !Target.IsMonolithic);
+			return UnrealBuildTool.IsEngineInstalled() || (Target.Type != TargetType.Program && !Target.IsMonolithic);
 		}
 
 		/// <summary>
@@ -1486,21 +1486,6 @@ namespace UnrealBuildTool
 		{
 			// Make sure the target file is known to us
 			bool bFoundTargetName = TargetNameToTargetFile.ContainsKey(TargetName);
-			if (bFoundTargetName == false)
-			{
-				if (UnrealBuildTool.RunningRocket())
-				{
-					//@todo Rocket: Remove this when full game support is implemented
-					// If we are Rocket, they will currently only have an editor target.
-					// See if that exists
-					bFoundTargetName = TargetNameToTargetFile.ContainsKey(TargetName + "Editor");
-					if (bFoundTargetName)
-					{
-						TargetName += "Editor";
-					}
-				}
-			}
-
 			if (bFoundTargetName == false)
 			{
 				if (Parent == null)
@@ -1681,23 +1666,6 @@ namespace UnrealBuildTool
 		/// We cache these file names so we can avoid searching for them later on.
 		static Dictionary<DirectoryReference, RulesFileCache> RootFolderToRulesFileCache = new Dictionary<DirectoryReference, RulesFileCache>();
 		
-		// Included for compatibility during //UE4/Main import
-		[Obsolete]
-		public static string GetModuleFilename(string ModuleName)
-		{
-			foreach(RulesFileCache Cache in RootFolderToRulesFileCache.Values)
-			{
-				foreach(FileReference ModuleFile in Cache.ModuleRules)
-				{
-					if(String.Compare(ModuleFile.GetFileNameWithoutAnyExtensions(), ModuleName, StringComparison.InvariantCultureIgnoreCase) == 0)
-					{
-						return ModuleFile.FullName;
-					}
-				}
-			}
-			return null;
-		}
-
 		public static List<FileReference> FindAllRulesSourceFiles(RulesFileType RulesFileType, List<DirectoryReference> GameFolders, List<FileReference> ForeignPlugins, List<DirectoryReference> AdditionalSearchPaths, bool bIncludeEngine = true)
 		{
 			List<DirectoryReference> Folders = new List<DirectoryReference>();
diff --git a/Engine/Source/Programs/UnrealBuildTool/System/UnrealBuildTool.cs b/Engine/Source/Programs/UnrealBuildTool/System/UnrealBuildTool.cs
--- a/Engine/Source/Programs/UnrealBuildTool/System/UnrealBuildTool.cs
+++ b/Engine/Source/Programs/UnrealBuildTool/System/UnrealBuildTool.cs
@@ -2430,8 +2430,8 @@ namespace UnrealBuildTool
 			//			-> UPDATE: Seems to work OK right now though on Windows platform, maybe due to GUID changes
 			// @todo ubtmake: Some platforms may not save any files into this folder.  We should delete + generate a "touch" file to force the directory timestamp to be updated (or just check the timestamp file itself.  We could put it ANYWHERE, actually)
 
-			// Rocket doesn't need to check engine projects for outdatedness
-			if (!UnrealBuildTool.RunningRocket())
+			// Installed Build doesn't need to check engine projects for outdatedness
+			if (!UnrealBuildTool.IsEngineInstalled())
 			{
 				if (ProjectFileGenerator.IntermediateProjectFilesPath.Exists())
 				{
diff --git a/Engine/Source/Programs/UnrealBuildTool/System/VCProject.cs b/Engine/Source/Programs/UnrealBuildTool/System/VCProject.cs
--- a/Engine/Source/Programs/UnrealBuildTool/System/VCProject.cs
+++ b/Engine/Source/Programs/UnrealBuildTool/System/VCProject.cs
@@ -203,11 +203,10 @@ namespace UnrealBuildTool
 			{
 				UnrealBuildTool.GetAllPlatforms(ref SupportedPlatforms);
 			}
-			bool bIncludeTestAndShippingConfigs = ProjectFileGenerator.bIncludeTestAndShippingConfigs || ProjectFileGenerator.bGeneratingRocketProjectFiles;
+
 			if (ProjectTarget.TargetRules != null)
 			{
-				// Rocket projects always get shipping configs
-				ProjectTarget.TargetRules.GetSupportedConfigurations(ref SupportedConfigurations, bIncludeTestAndShippingConfigs: bIncludeTestAndShippingConfigs);
+				ProjectTarget.TargetRules.GetSupportedConfigurations(ref SupportedConfigurations, ProjectFileGenerator.bIncludeTestAndShippingConfigs);
 			}
 
 			// Add all of the extra platforms/configurations for this target
@@ -221,7 +220,7 @@ namespace UnrealBuildTool
 				}
 				foreach (var ExtraConfig in ProjectTarget.ExtraSupportedConfigurations)
 				{
-					if (bIncludeTestAndShippingConfigs || (ExtraConfig != UnrealTargetConfiguration.Shipping && ExtraConfig != UnrealTargetConfiguration.Test))
+					if (ProjectFileGenerator.bIncludeTestAndShippingConfigs || (ExtraConfig != UnrealTargetConfiguration.Shipping && ExtraConfig != UnrealTargetConfiguration.Test))
 					{
 						if (!SupportedConfigurations.Contains(ExtraConfig))
 						{
@@ -631,8 +630,8 @@ namespace UnrealBuildTool
 					"	</ItemGroup>" + ProjectFileGenerator.NewLine);
 			}
 
-			// For Rocket, include engine source in the source search paths. We never build it locally, so the debugger can't find it.
-			if (UnrealBuildTool.RunningRocket() && !IsStubProject)
+			// For Installed engine builds, include engine source in the source search paths if it exists. We never build it locally, so the debugger can't find it.
+			if (UnrealBuildTool.IsEngineInstalled() && !IsStubProject)
 			{
 				VCProjectFileContent.Append("	<PropertyGroup>" + ProjectFileGenerator.NewLine);
 				VCProjectFileContent.Append("		<SourcePath>");
@@ -964,7 +963,7 @@ namespace UnrealBuildTool
 						"		<NMakeOutput/>" + ProjectFileGenerator.NewLine +
 						"	</PropertyGroup>" + ProjectFileGenerator.NewLine);
 				}
-				else if (ProjectFileGenerator.bGeneratingRocketProjectFiles && Combination.ProjectTarget != null && Combination.ProjectTarget.TargetRules != null && !Combination.ProjectTarget.TargetRules.SupportsPlatform(Combination.Platform))
+				else if (UnrealBuildTool.IsEngineInstalled() && Combination.ProjectTarget != null && Combination.ProjectTarget.TargetRules != null && !Combination.ProjectTarget.TargetRules.SupportsPlatform(Combination.Platform))
 				{
 					List<UnrealTargetPlatform> SupportedPlatforms = new List<UnrealTargetPlatform>();
 					Combination.ProjectTarget.TargetRules.GetSupportedPlatforms(ref SupportedPlatforms);
diff --git a/Engine/Source/Programs/UnrealBuildTool/System/XcodeProject.cs b/Engine/Source/Programs/UnrealBuildTool/System/XcodeProject.cs
--- a/Engine/Source/Programs/UnrealBuildTool/System/XcodeProject.cs
+++ b/Engine/Source/Programs/UnrealBuildTool/System/XcodeProject.cs
@@ -426,7 +426,7 @@ namespace UnrealBuildTool
 		private void AppendLegacyTargetSection(StringBuilder Content, string TargetName, string TargetGuid, string TargetBuildConfigGuid, FileReference UProjectPath)
 		{
 			string UE4Dir = ConvertPath(Path.GetFullPath(Directory.GetCurrentDirectory() + "../../.."));
-			string BuildToolPath = UE4Dir + "/Engine/Build/BatchFiles/Mac/" + (XcodeProjectFileGenerator.bGeneratingRocketProjectFiles ? "Rocket" : "") + "Build.sh";
+			string BuildToolPath = UE4Dir + "/Engine/Build/BatchFiles/Mac/" + (UnrealBuildTool.RunningRocket() ? "Rocket" : "") + "Build.sh";
 
 			Content.Append("/* Begin PBXLegacyTarget section */" + ProjectFileGenerator.NewLine);
 
diff --git a/Engine/Source/Programs/UnrealBuildTool/ToolChain/UEToolChain.cs b/Engine/Source/Programs/UnrealBuildTool/ToolChain/UEToolChain.cs
--- a/Engine/Source/Programs/UnrealBuildTool/ToolChain/UEToolChain.cs
+++ b/Engine/Source/Programs/UnrealBuildTool/ToolChain/UEToolChain.cs
@@ -9,14 +9,7 @@ using Microsoft.Win32;
 
 namespace UnrealBuildTool
 {
-	// Included for compatibility during //UE4/Main import
-	public interface IUEToolChain
-	{
-		[Obsolete]
-		void StripSymbols(string SourceFile, string TargetFile);
-	}
-
-	public abstract class UEToolChain : IUEToolChain
+	public abstract class UEToolChain
 	{
 		public readonly CPPTargetPlatform CppPlatform;
 
@@ -25,15 +18,6 @@ namespace UnrealBuildTool
 			CppPlatform = InCppPlatform;
 		}
 
-		// Included for compatibility during //UE4/Main import
-		[Obsolete]
-		public static IUEToolChain GetPlatformToolChain(CPPTargetPlatform Platform)
-		{
-			UEBuildPlatform BuildPlatform = UEBuildPlatform.GetBuildPlatform(UEBuildTarget.CPPTargetPlatformToUnrealTargetPlatform(Platform));
-			UEBuildPlatformContext Context = BuildPlatform.CreateContext(null);
-			return Context.CreateToolChain(BuildPlatform.DefaultCppPlatform);
-		}
-
 		public abstract CPPOutput CompileCPPFiles(UEBuildTarget Target, CPPEnvironment CompileEnvironment, List<FileItem> SourceFiles, string ModuleName);
 
 		public virtual CPPOutput CompileRCFiles(UEBuildTarget Target, CPPEnvironment Environment, List<FileItem> RCFiles)
diff --git a/Engine/Source/Programs/UnrealGameSync/UnrealGameSync/Controls/LogSplitContainer.cs b/Engine/Source/Programs/UnrealGameSync/UnrealGameSync/Controls/LogSplitContainer.cs
--- a/Engine/Source/Programs/UnrealGameSync/UnrealGameSync/Controls/LogSplitContainer.cs
+++ b/Engine/Source/Programs/UnrealGameSync/UnrealGameSync/Controls/LogSplitContainer.cs
@@ -40,6 +40,13 @@ namespace UnrealGameSync
 			get { return SystemFonts.IconTitleFont; }
 		}
 
+		protected override void OnCreateControl()
+		{
+			base.OnCreateControl();
+
+			UpdateMetrics();
+		}
+
 		protected override void OnKeyUp(KeyEventArgs e)
 		{
 			// Base implementation directly repaints the splitter control in response to this message, which bypasses our custom paint handler. Invalidate the window to get it drawn again.
diff --git a/Engine/Source/Programs/UnrealGameSync/UnrealGameSync/Controls/StatusPanel.cs b/Engine/Source/Programs/UnrealGameSync/UnrealGameSync/Controls/StatusPanel.cs
--- a/Engine/Source/Programs/UnrealGameSync/UnrealGameSync/Controls/StatusPanel.cs
+++ b/Engine/Source/Programs/UnrealGameSync/UnrealGameSync/Controls/StatusPanel.cs
@@ -272,12 +272,23 @@ namespace UnrealGameSync
 		Point? MouseDownLocation;
 		StatusElement MouseDownElement;
 		int ContentWidth = 400;
+		int SuspendDisplayCount;
 
 		public StatusPanel()
 		{
 			DoubleBuffered = true;
 		}
 
+		public void SuspendDisplay()
+		{
+			SuspendDisplayCount++;
+		}
+
+		public void ResumeDisplay()
+		{
+			SuspendDisplayCount--;
+		}
+
 		public void SetContentWidth(int NewContentWidth)
 		{
 			if(ContentWidth != NewContentWidth)
@@ -462,13 +473,16 @@ namespace UnrealGameSync
 		{
 			base.OnPaint(e);
 
-			e.Graphics.DrawImage(ProjectLogo ?? Properties.Resources.DefaultProjectLogo, ProjectLogoBounds);
+			if(SuspendDisplayCount == 0)
+			{
+				e.Graphics.DrawImage(ProjectLogo ?? Properties.Resources.DefaultProjectLogo, ProjectLogoBounds);
 
-			e.Graphics.SmoothingMode = System.Drawing.Drawing2D.SmoothingMode.None;
+				e.Graphics.SmoothingMode = System.Drawing.Drawing2D.SmoothingMode.None;
 
-			foreach(StatusLine Line in Lines)
-			{
-				Line.Draw(e.Graphics, FontCache);
+				foreach(StatusLine Line in Lines)
+				{
+					Line.Draw(e.Graphics, FontCache);
+				}
 			}
 		}
 
diff --git a/Engine/Source/Programs/UnrealGameSync/UnrealGameSync/DetectProjectSettingsTask.cs b/Engine/Source/Programs/UnrealGameSync/UnrealGameSync/DetectProjectSettingsTask.cs
--- a/Engine/Source/Programs/UnrealGameSync/UnrealGameSync/DetectProjectSettingsTask.cs
+++ b/Engine/Source/Programs/UnrealGameSync/UnrealGameSync/DetectProjectSettingsTask.cs
@@ -19,6 +19,7 @@ namespace UnrealGameSync
 		public string BranchClientPath;
 		public string BaseEditorTargetPath;
 		public string NewSelectedClientFileName;
+		public string StreamName;
 		public Image ProjectLogo;
 		public TimeSpan ServerTimeZone;
 		TextWriter Log;
@@ -186,7 +187,6 @@ namespace UnrealGameSync
 			}
 
 			// Get a unique name for the project that's selected. For regular branches, this can be the depot path. For streams, we want to include the stream name to encode imports.
-			string StreamName;
 			if(PerforceClient.GetActiveStream(out StreamName, Log))
 			{
 				string ExpectedPrefix = String.Format("//{0}/", PerforceClient.ClientName);
diff --git a/Engine/Source/Programs/UnrealGameSync/UnrealGameSync/Forms/MainWindow.cs b/Engine/Source/Programs/UnrealGameSync/UnrealGameSync/Forms/MainWindow.cs
--- a/Engine/Source/Programs/UnrealGameSync/UnrealGameSync/Forms/MainWindow.cs
+++ b/Engine/Source/Programs/UnrealGameSync/UnrealGameSync/Forms/MainWindow.cs
@@ -64,6 +64,7 @@ namespace UnrealGameSync
 		string SelectedFileName;
 		string SelectedProjectIdentifier;
 		string BranchDirectoryName;
+		string StreamName;
 		string EditorTargetName;
 		PerforceMonitor PerforceMonitor;
 		Workspace Workspace;
@@ -377,6 +378,7 @@ namespace UnrealGameSync
 				SelectedProjectIdentifier = DetectSettings.NewSelectedProjectIdentifier;
 				EditorTargetName = DetectSettings.NewProjectEditorTarget;
 				BranchDirectoryName = Path.GetFullPath(Path.Combine(Path.GetDirectoryName(DetectSettings.BaseEditorTargetPath), "..", ".."));
+				StreamName = DetectSettings.StreamName;
 				ServerTimeZone = DetectSettings.ServerTimeZone;
 
 				// Check if we've the project we've got open in this workspace is the one we're actually synced to
@@ -440,11 +442,15 @@ namespace UnrealGameSync
 
 		private void StreamChangedCallback()
 		{
+			StatusPanel.SuspendDisplay();
+
 			string PrevSelectedFileName = SelectedFileName;
 			if(TryCloseProject())
 			{
 				OpenProject(PrevSelectedFileName);
 			}
+
+			StatusPanel.ResumeDisplay();
 		}
 
 		private void StreamChangedCallbackAsync()
@@ -1633,7 +1639,7 @@ namespace UnrealGameSync
 
 		private void UpdateStatusPanel()
 		{
-			int NewContentWidth = Math.Max(TextRenderer.MeasureText("Last synced to changelist 123456789. 12 users are on a newer build.        | Sync Now | Sync To...", StatusPanel.Font).Width, 400);
+			int NewContentWidth = Math.Max(TextRenderer.MeasureText(String.Format("Last synced to {0} at changelist 123456789. 12 users are on a newer build.        | Sync Now | Sync To...", StreamName ?? ""), StatusPanel.Font).Width, 400);
 			StatusPanel.SetContentWidth(NewContentWidth);
 
 			List<StatusLine> Lines = new List<StatusLine>();
@@ -1672,7 +1678,16 @@ namespace UnrealGameSync
 				StatusLine SummaryLine = new StatusLine();
 				if(Workspace.CurrentChangeNumber != -1)
 				{
-					SummaryLine.AddText("Last synced to changelist ");
+					if(StreamName == null)
+					{
+						SummaryLine.AddText("Last synced to changelist ");
+					}
+					else
+					{
+						SummaryLine.AddText("Last synced to ");
+						SummaryLine.AddLink(StreamName + "\u25BE", FontStyle.Regular, (P, R) => { SelectOtherStream(R); });
+						SummaryLine.AddText(" at changelist ");
+					}
 					SummaryLine.AddLink(String.Format("{0}.", Workspace.CurrentChangeNumber), FontStyle.Regular, () => { SelectChange(Workspace.CurrentChangeNumber); });
 					int NumUsersOnNewerChange = SortedChangeNumbers.Where(x => (x > Workspace.CurrentChangeNumber)).OrderByDescending(x => x).Select(x => EventMonitor.GetSummaryForChange(x)).Where(x => x != null && (!ShouldSyncPrecompiledEditor || GetArchivePathForChangeNumber(x.ChangeNumber) != null)).Sum(x => x.CurrentUsers.Count);
 					if(NumUsersOnNewerChange > 0)
@@ -1682,7 +1697,15 @@ namespace UnrealGameSync
 				}
 				else
 				{
-					SummaryLine.AddText("You are not currently synced to this branch.");
+					SummaryLine.AddText("You are not currently synced to ");
+					if(StreamName == null)
+					{
+						SummaryLine.AddText("this branch.");
+					}
+					else
+					{
+						SummaryLine.AddLink(StreamName + " \u25BE", FontStyle.Regular, (P, R) => { SelectOtherStream(R); });
+					}
 				}
 				SummaryLine.AddText("  |  ");
 				SummaryLine.AddLink("Sync Now", FontStyle.Bold | FontStyle.Underline, () => { SyncLatestChange(); });
@@ -1751,6 +1774,42 @@ namespace UnrealGameSync
 			StatusPanel.Set(Lines);
 		}
 
+		private void SelectOtherStream(Rectangle Bounds)
+		{
+			if(StreamName != null)
+			{
+				IReadOnlyList<string> OtherStreamNames = PerforceMonitor.OtherStreamNames;
+
+				StreamContextMenu.Items.Clear();
+				foreach(string OtherStreamName in OtherStreamNames.OrderBy(x => x).Where(x => !x.EndsWith("/Dev-Binaries")))
+				{
+					string ThisStreamName = OtherStreamName; // Local for lambda capture
+
+					ToolStripMenuItem Item = new ToolStripMenuItem(ThisStreamName, null, new EventHandler((S, E) => SelectStream(ThisStreamName)));
+					if(String.Compare(StreamName, OtherStreamName, StringComparison.InvariantCultureIgnoreCase) == 0)
+					{
+						Item.Checked = true;
+					}
+					StreamContextMenu.Items.Add(Item);
+				}
+				StreamContextMenu.Show(StatusPanel, new Point(Bounds.Left, Bounds.Bottom), ToolStripDropDownDirection.BelowRight);
+			}
+		}
+
+		private void SelectStream(string StreamName)
+		{
+			if(Workspace.IsBusy())
+			{
+				MessageBox.Show("Please retry after the current sync has finished.", "Sync in Progress");
+			}
+			else if(Workspace.Perforce != null && Workspace.Perforce.SwitchStream(StreamName, Log))
+			{
+				StatusPanel.SuspendLayout();
+				StreamChangedCallback();
+				StatusPanel.ResumeLayout();
+			}
+		}
+
 		private void ViewLastSyncStatus()
 		{
 			string SummaryText;
@@ -1985,6 +2044,11 @@ namespace UnrealGameSync
 			}
 		}
 
+		private void ProjectList_MouseWheel(object sender, MouseEventArgs e)
+		{
+			((HandledMouseEventArgs)e).Handled = true;
+		}
+
 		private void ProjectList_SelectionChangeCommitted(object sender, EventArgs e)
 		{
 			if(ProjectList.SelectedIndex >= 0 && ProjectList.SelectedIndex == ProjectList.Items.Count - 1)
diff --git a/Engine/Source/Programs/UnrealGameSync/UnrealGameSync/Perforce.cs b/Engine/Source/Programs/UnrealGameSync/UnrealGameSync/Perforce.cs
--- a/Engine/Source/Programs/UnrealGameSync/UnrealGameSync/Perforce.cs
+++ b/Engine/Source/Programs/UnrealGameSync/UnrealGameSync/Perforce.cs
@@ -307,6 +307,7 @@ namespace UnrealGameSync
 			IgnoreNoSuchFilesError = 0x10,
 			IgnoreFilesNotOpenedOnThisClientError = 0x20,
 			IgnoreExitCode = 0x40,
+			IgnoreFilesNotInClientViewError = 0x80,
 		}
 
 		delegate bool HandleRecordDelegate(Dictionary<string, string> Tags);
@@ -461,7 +462,7 @@ namespace UnrealGameSync
 		public bool FileExists(string Filter, out bool bExists, TextWriter Log)
 		{
 			List<PerforceFileRecord> FileRecords;
-			if(RunCommand(String.Format("fstat \"{0}\"", Filter), out FileRecords, CommandOptions.IgnoreNoSuchFilesError, Log))
+			if(RunCommand(String.Format("fstat \"{0}\"", Filter), out FileRecords, CommandOptions.IgnoreNoSuchFilesError | CommandOptions.IgnoreFilesNotInClientViewError, Log))
 			{
 				bExists = (FileRecords.Count > 0);
 				return true;
@@ -687,6 +688,37 @@ namespace UnrealGameSync
 			}
 		}
 
+		public bool FindStreams(string Filter, out List<string> OutStreamNames, TextWriter Log)
+		{
+			List<string> Lines;
+			if(!RunCommand(String.Format("streams -F \"Stream={0}\"", Filter), out Lines, CommandOptions.None, Log))
+			{
+				OutStreamNames = null;
+				return false;
+			}
+
+			List<string> StreamNames = new List<string>();
+			foreach(string Line in Lines)
+			{
+				string[] Tokens = Line.Split(new char[]{ ' ' }, StringSplitOptions.RemoveEmptyEntries);
+				if(Tokens.Length < 2 || Tokens[0] != "Stream" || !Tokens[1].StartsWith("//"))
+				{
+					Log.WriteLine("Unexpected output from stream query: {0}", Line);
+					OutStreamNames = null;
+					return false;
+				}
+				StreamNames.Add(Tokens[1]);
+			}
+			OutStreamNames = StreamNames;
+
+			return true;
+		}
+
+		public bool SwitchStream(string NewStream, TextWriter Log)
+		{
+			return RunCommand(String.Format("client -f -s -S \"{0}\" \"{1}\"", NewStream, ClientName), CommandOptions.None, Log);
+		}
+
 		public bool Where(string Filter, out PerforceWhereRecord WhereRecord, TextWriter Log)
 		{
 			List<PerforceFileRecord> FileRecords;
@@ -1018,6 +1050,10 @@ namespace UnrealGameSync
 			{
 				return true;
 			}
+			else if(Options.HasFlag(CommandOptions.IgnoreFilesNotInClientViewError) && Text.StartsWith("error: ") && Text.EndsWith("- file(s) not in client view."))
+			{
+				return true;
+			}
 			else if(Options.HasFlag(CommandOptions.IgnoreFilesNotOpenedOnThisClientError) && Text.StartsWith("error: ") && Text.EndsWith(" - file(s) not opened on this client."))
 			{
 				return true;
diff --git a/Engine/Source/Programs/UnrealGameSync/UnrealGameSync/PerforceMonitor.cs b/Engine/Source/Programs/UnrealGameSync/UnrealGameSync/PerforceMonitor.cs
--- a/Engine/Source/Programs/UnrealGameSync/UnrealGameSync/PerforceMonitor.cs
+++ b/Engine/Source/Programs/UnrealGameSync/UnrealGameSync/PerforceMonitor.cs
@@ -56,6 +56,7 @@ namespace UnrealGameSync
 			SelectedProjectIdentifier = InSelectedProjectIdentifier;
 			PendingMaxChangesValue = 100;
 			LastCodeChangeByCurrentUser = -1;
+			OtherStreamNames = new List<string>();
 
 			LogWriter = new BoundedLogWriter(InLogPath);
 
@@ -100,6 +101,12 @@ namespace UnrealGameSync
 			set { lock(this){ if(value != PendingMaxChangesValue){ PendingMaxChangesValue = value; RefreshEvent.Set(); } } }
 		}
 
+		public IReadOnlyList<string> OtherStreamNames
+		{
+			get;
+			private set;
+		}
+
 		void PollForUpdates()
 		{
 			string StreamName;
@@ -122,6 +129,17 @@ namespace UnrealGameSync
 					OnStreamChange();
 				}
 
+				// Update the stream list
+				if(StreamName != null)
+				{
+					List<string> NewOtherStreamNames;
+					if(!Perforce.FindStreams(PerforceUtils.GetClientOrDepotDirectoryName(StreamName) + "/*", out NewOtherStreamNames, LogWriter))
+					{
+						NewOtherStreamNames = new List<string>();
+					}
+					OtherStreamNames = NewOtherStreamNames;
+				}
+
 				// Check for any p4 changes
 				if(!UpdateChanges())
 				{
@@ -384,6 +402,12 @@ namespace UnrealGameSync
 			}
 		}
 
+		public string CurrentStream
+		{
+			get;
+			private set;
+		}
+
 		public int LastCodeChangeByCurrentUser
 		{
 			get;
diff --git a/Engine/Source/Programs/UnrealGameSync/UnrealGameSync/Properties/AssemblyInfo.cs b/Engine/Source/Programs/UnrealGameSync/UnrealGameSync/Properties/AssemblyInfo.cs
--- a/Engine/Source/Programs/UnrealGameSync/UnrealGameSync/Properties/AssemblyInfo.cs
+++ b/Engine/Source/Programs/UnrealGameSync/UnrealGameSync/Properties/AssemblyInfo.cs
@@ -34,5 +34,5 @@ using System.Runtime.InteropServices;
 // You can specify all the values or you can default the Build and Revision Numbers 
 // by using the '*' as shown below:
 // [assembly: AssemblyVersion("1.0.*")]
-[assembly: AssemblyVersion("1.63.*")]
-[assembly: AssemblyFileVersion("1.63.0.0")]
+[assembly: AssemblyVersion("1.65.*")]
+[assembly: AssemblyFileVersion("1.65.0.0")]
diff --git a/Engine/Source/Runtime/Analytics/Analytics/Public/Analytics.h b/Engine/Source/Runtime/Analytics/Analytics/Public/Analytics.h
--- a/Engine/Source/Runtime/Analytics/Analytics/Public/Analytics.h
+++ b/Engine/Source/Runtime/Analytics/Analytics/Public/Analytics.h
@@ -208,14 +208,6 @@ public:
 			return FName(*ProviderModuleName);
 		}
 
-		/**
-		 * Tries to get an override value for the UE4Type from config files
-		 */
-		static bool GetUE4TypeOverride(FString& OutUE4Type)
-		{
-			return GConfig->GetString(TEXT("Analytics"), TEXT("UE4TypeOverride"), OutUE4Type, GEngineIni);
-		}
-
 		/** Allows setting the INI section name based on the build type passed in. Allows access to the default section values when the application chooses the build type itself. */
 		void SetSectionNameByBuildType(FAnalytics::BuildType InBuildType)
 		{
diff --git a/Engine/Source/Runtime/Analytics/AnalyticsET/Private/AnalyticsET.cpp b/Engine/Source/Runtime/Analytics/AnalyticsET/Private/AnalyticsET.cpp
--- a/Engine/Source/Runtime/Analytics/AnalyticsET/Private/AnalyticsET.cpp
+++ b/Engine/Source/Runtime/Analytics/AnalyticsET/Private/AnalyticsET.cpp
@@ -54,8 +54,6 @@ private:
 	FString UserID;
 	/** The session ID */
 	FString SessionID;
-	/** Cached build type as a string */
-	FString BuildType;
 	/** The AppVersion passed to ET. */
 	FString AppVersion;
 	/** Max number of analytics events to cache before pushing to server */
@@ -70,6 +68,12 @@ private:
 	bool bInDestructor;
 	/** True to use the legacy backend server protocol that uses URL params. */
 	bool UseLegacyProtocol;
+	/** AppEnvironment to use. */
+	FString AppEnvironment;
+	/** UploadType to use. */
+	FString UploadType;
+	/** Tag as to whether the provider was given a legacy AppServer endpoint instead of a newer data router one. Used to find apps that need to be updated in analytics. */
+	bool bUsingLegacyAppServer;
 	/**
 	 * Analytics event entry to be cached
 	 */
@@ -118,9 +122,14 @@ TSharedPtr<IAnalyticsProvider> FAnalyticsET::CreateAnalyticsProvider(const FAnal
 	{
 		Config ConfigValues;
 		ConfigValues.APIKeyET = GetConfigValue.Execute(Config::GetKeyNameForAPIKey(), true);
-		ConfigValues.APIServerET = GetConfigValue.Execute(Config::GetKeyNameForAPIServer(), false);
+		ConfigValues.APIServerET = GetConfigValue.Execute(Config::GetKeyNameForAPIServer(), true);
 		ConfigValues.AppVersionET = GetConfigValue.Execute(Config::GetKeyNameForAppVersion(), false);
 		ConfigValues.UseLegacyProtocol = FCString::ToBool(*GetConfigValue.Execute(Config::GetKeyNameForUseLegacyProtocol(), false));
+		if (!ConfigValues.UseLegacyProtocol)
+		{
+			ConfigValues.AppEnvironment = GetConfigValue.Execute(Config::GetKeyNameForAppEnvironment(), true);
+			ConfigValues.UploadType = GetConfigValue.Execute(Config::GetKeyNameForUploadType(), true);
+		}
 		return CreateAnalyticsProvider(ConfigValues);
 	}
 	else
@@ -146,27 +155,29 @@ TSharedPtr<IAnalyticsProvider> FAnalyticsET::CreateAnalyticsProvider(const Confi
  */
 FAnalyticsProviderET::FAnalyticsProviderET(const FAnalyticsET::Config& ConfigValues)
 	:bSessionInProgress(false)
+	, APIKey(ConfigValues.APIKeyET)
+	, APIServer(ConfigValues.APIServerET)
 	, MaxCachedNumEvents(20)
 	, MaxCachedElapsedTime(60.0f)
 	, bShouldCacheEvents(!FParse::Param(FCommandLine::Get(), TEXT("ANALYTICSDISABLECACHING")))
 	, FlushEventsCountdown(MaxCachedElapsedTime)
 	, bInDestructor(false)
 	, UseLegacyProtocol(ConfigValues.UseLegacyProtocol)
+	, bUsingLegacyAppServer(false)
 {
+	if (APIKey.IsEmpty() || APIServer.IsEmpty())
+	{
+		UE_LOG(LogAnalytics, Fatal, TEXT("AnalyticsET: APIKey (%s) and APIServer (%s) cannot be empty!"), *APIKey, *APIServer);
+	}
+
 	// if we are not caching events, we are operating in debug mode. Turn on super-verbose analytics logging
 	if (!bShouldCacheEvents)
 	{
 		UE_SET_LOG_VERBOSITY(LogAnalytics, VeryVerbose);
 	}
 
-	APIKey = ConfigValues.APIKeyET;
 	UE_LOG(LogAnalytics, Verbose, TEXT("[%s] Initializing ET Analytics provider"), *APIKey);
 
-	// allow the APIServer value to be empty and use defaults.
-	APIServer = ConfigValues.APIServerET.IsEmpty() 
-		? FAnalyticsET::Config::GetDefaultAPIServer()
-		: ConfigValues.APIServerET;
-
 	// default to FEngineVersion::Current() if one is not provided, append FEngineVersion::Current() otherwise.
 	FString ConfigAppVersion = ConfigValues.AppVersionET;
 	// Allow the cmdline to force a specific AppVersion so it can be set dynamically.
@@ -175,20 +186,28 @@ FAnalyticsProviderET::FAnalyticsProviderET(const FAnalyticsET::Config& ConfigVal
 		? FEngineVersion::Current().ToString() 
 		: ConfigAppVersion.Replace(TEXT("%VERSION%"), *FEngineVersion::Current().ToString(), ESearchCase::CaseSensitive);
 
+	// support legacy connections that send to the old URL and route them to the new URL.
+	// This automatically route older titles to the new endpoint until all their configs can be updated.
+	// @todo: Remove this hack once all internal titles have been updated to use the new URL explicitly.
+	if (APIServer.Equals(TEXT("http://et2.epicgames.com/ET2/"), ESearchCase::IgnoreCase))
+	{
+		UE_LOG(LogAnalytics, Warning, TEXT("[%s] APIServer is using legacy URL (%s). Automatically updating to the new data router endpoint. Please update your configuration to explicitly point to a data router endpoint."), *APIKey, *APIServer);
+		bUsingLegacyAppServer = true;
+		APIServer = TEXT("https://datarouter.ol.epicgames.com/");
+	}
+
 	UE_LOG(LogAnalytics, Log, TEXT("[%s] APIServer = %s. AppVersion = %s"), *APIKey, *APIServer, *AppVersion);
 
-	// cache the build type string
-	FAnalytics::BuildType BuildTypeEnum = FAnalytics::Get().GetBuildType();
-	BuildType = 
-		BuildTypeEnum == FAnalytics::Debug 
-			? "Debug"
-			: BuildTypeEnum == FAnalytics::Development
-				? "Development"
-				: BuildTypeEnum == FAnalytics::Release
-					? "Release"
-					: BuildTypeEnum == FAnalytics::Test
-						? "Test"
-						: "UNKNOWN";
+	// only need these if we are using the data router protocol.
+	if (!UseLegacyProtocol)
+	{
+		AppEnvironment = ConfigValues.AppEnvironment.IsEmpty()
+			? FAnalyticsET::Config::GetDefaultAppEnvironment()
+			: ConfigValues.AppEnvironment;
+		UploadType = ConfigValues.UploadType.IsEmpty()
+			? FAnalyticsET::Config::GetDefaultUploadType()
+			: ConfigValues.UploadType;
+	}
 
 	// see if there is a cmdline supplied UserID.
 #if !UE_BUILD_SHIPPING
@@ -249,6 +268,7 @@ bool FAnalyticsProviderET::StartSession(const TArray<FAnalyticsEventAttribute>&
 	AppendedAttributes.Emplace(TEXT("UniqueDeviceId"), FPlatformMisc::GetUniqueDeviceId());
 	// we should always know what platform is hosting this session.
 	AppendedAttributes.Emplace(TEXT("Platform"), FString(FPlatformProperties::IniPlatformName()));
+	AppendedAttributes.Emplace(TEXT("LegacyURL"), bUsingLegacyAppServer ? TEXT("true") : TEXT("false"));
 
 	RecordEvent(TEXT("SessionStart"), AppendedAttributes);
 	bSessionInProgress = true;
@@ -290,62 +310,66 @@ void FAnalyticsProviderET::FlushEvents()
 
 		if (!UseLegacyProtocol)
 		{
-		TSharedRef< TJsonWriter<TCHAR, TCondensedJsonPrintPolicy<TCHAR> > > JsonWriter = TJsonWriterFactory<TCHAR, TCondensedJsonPrintPolicy<TCHAR> >::Create(&Payload);
-		JsonWriter->WriteObjectStart();
-		JsonWriter->WriteArrayStart(TEXT("Events"));
-		for (int32 EventIdx = 0; EventIdx < CachedEvents.Num(); EventIdx++)
-		{
-			const FAnalyticsEventEntry& Entry = CachedEvents[EventIdx];
-			// event entry
+			TSharedRef< TJsonWriter<TCHAR, TCondensedJsonPrintPolicy<TCHAR> > > JsonWriter = TJsonWriterFactory<TCHAR, TCondensedJsonPrintPolicy<TCHAR> >::Create(&Payload);
 			JsonWriter->WriteObjectStart();
-			JsonWriter->WriteValue(TEXT("EventName"), Entry.EventName);
-			FString DateOffset = (CurrentTime - Entry.TimeStamp).ToString();
-			JsonWriter->WriteValue(TEXT("DateOffset"), DateOffset);
-			JsonWriter->WriteValue(TEXT("IsEditor"), FString::FromInt(GIsEditor));
-			if (Entry.Attributes.Num() > 0)
+			JsonWriter->WriteArrayStart(TEXT("Events"));
+			for (int32 EventIdx = 0; EventIdx < CachedEvents.Num(); EventIdx++)
 			{
-				// optional attributes for this event
-				for (int32 AttrIdx = 0; AttrIdx < Entry.Attributes.Num(); AttrIdx++)
+				const FAnalyticsEventEntry& Entry = CachedEvents[EventIdx];
+				// event entry
+				JsonWriter->WriteObjectStart();
+				JsonWriter->WriteValue(TEXT("EventName"), Entry.EventName);
+				FString DateOffset = (CurrentTime - Entry.TimeStamp).ToString();
+				JsonWriter->WriteValue(TEXT("DateOffset"), DateOffset);
+				JsonWriter->WriteValue(TEXT("IsEditor"), FString::FromInt(GIsEditor));
+				if (Entry.Attributes.Num() > 0)
 				{
+					// optional attributes for this event
+					for (int32 AttrIdx = 0; AttrIdx < Entry.Attributes.Num(); AttrIdx++)
+					{
 						const FAnalyticsEventAttribute& Attr = Entry.Attributes[AttrIdx];
-					JsonWriter->WriteValue(Attr.AttrName, Attr.AttrValue);
+						JsonWriter->WriteValue(Attr.AttrName, Attr.AttrValue);
+					}
 				}
+				JsonWriter->WriteObjectEnd();
 			}
+			JsonWriter->WriteArrayEnd();
 			JsonWriter->WriteObjectEnd();
-		}
-		JsonWriter->WriteArrayEnd();
-		JsonWriter->WriteObjectEnd();
-		JsonWriter->Close();
-
-		FString URLPath = FString::Printf(TEXT("CollectData.1?SessionID=%s&AppID=%s&AppVersion=%s&UserID=%s"),
-			*FPlatformHttp::UrlEncode(SessionID),
-			*FPlatformHttp::UrlEncode(APIKey),
-			*FPlatformHttp::UrlEncode(AppVersion),
-			*FPlatformHttp::UrlEncode(UserID));
-
-		// Recreate the URLPath for logging because we do not want to escape the parameters when logging.
-		// We cannot simply UrlEncode the entire Path after logging it because UrlEncode(Params) != UrlEncode(Param1) & UrlEncode(Param2) ...
-			UE_LOG(LogAnalytics, VeryVerbose, TEXT("[%s] AnalyticsET URL:CollectData.1?SessionID=%s&AppID=%s&AppVersion=%s&UserID=%s. Payload:%s"),
+			JsonWriter->Close();
+
+			FString URLPath = FString::Printf(TEXT("datarouter/api/v1/public/data?SessionID=%s&AppID=%s&AppVersion=%s&UserID=%s&AppEnvironment=%s&UploadType=%s"),
+				*FPlatformHttp::UrlEncode(SessionID),
+				*FPlatformHttp::UrlEncode(APIKey),
+				*FPlatformHttp::UrlEncode(AppVersion),
+				*FPlatformHttp::UrlEncode(UserID),
+				*FPlatformHttp::UrlEncode(AppEnvironment),
+				*FPlatformHttp::UrlEncode(UploadType));
+
+			// Recreate the URLPath for logging because we do not want to escape the parameters when logging.
+			// We cannot simply UrlEncode the entire Path after logging it because UrlEncode(Params) != UrlEncode(Param1) & UrlEncode(Param2) ...
+			UE_LOG(LogAnalytics, VeryVerbose, TEXT("[%s] AnalyticsET URL:datarouter/api/v1/public/data?SessionID=%s&AppID=%s&AppVersion=%s&UserID=%s&AppEnvironment=%s&UploadType=%s. Payload:%s"),
 				*APIKey,
-			*SessionID,
-			*APIKey,
-			*AppVersion,
-			*UserID,
-			*Payload);
-
-		// Create/send Http request for an event
-		TSharedRef<IHttpRequest> HttpRequest = FHttpModule::Get().CreateRequest();
-		HttpRequest->SetHeader(TEXT("Content-Type"), TEXT("application/json; charset=utf-8"));
-
-		HttpRequest->SetURL(APIServer + URLPath);
-		HttpRequest->SetVerb(TEXT("POST"));
-		HttpRequest->SetContentAsString(Payload);
-		// Don't set a response callback if we are in our destructor, as the instance will no longer be there to call.
-		if (!bInDestructor)
-		{
-			HttpRequest->OnProcessRequestComplete().BindSP(this, &FAnalyticsProviderET::EventRequestComplete);
-		}
- 		HttpRequest->ProcessRequest();
+				*SessionID,
+				*APIKey,
+				*AppVersion,
+				*UserID,
+				*AppEnvironment,
+				*UploadType,
+				*Payload);
+
+			// Create/send Http request for an event
+			TSharedRef<IHttpRequest> HttpRequest = FHttpModule::Get().CreateRequest();
+			HttpRequest->SetHeader(TEXT("Content-Type"), TEXT("application/json; charset=utf-8"));
+
+			HttpRequest->SetURL(APIServer + URLPath);
+			HttpRequest->SetVerb(TEXT("POST"));
+			HttpRequest->SetContentAsString(Payload);
+			// Don't set a response callback if we are in our destructor, as the instance will no longer be there to call.
+			if (!bInDestructor)
+			{
+				HttpRequest->OnProcessRequestComplete().BindSP(this, &FAnalyticsProviderET::EventRequestComplete);
+			}
+ 			HttpRequest->ProcessRequest();
 		}
 		else
 		{
diff --git a/Engine/Source/Runtime/Analytics/AnalyticsET/Private/HttpServiceTracker.cpp b/Engine/Source/Runtime/Analytics/AnalyticsET/Private/HttpServiceTracker.cpp
--- a/Engine/Source/Runtime/Analytics/AnalyticsET/Private/HttpServiceTracker.cpp
+++ b/Engine/Source/Runtime/Analytics/AnalyticsET/Private/HttpServiceTracker.cpp
@@ -67,6 +67,8 @@ FHttpServiceTracker::FHttpServiceTracker(const FHttpServiceTrackerConfig& Config
 	AnalyticsConfig.APIKeyET = Config.APIKey;
 	AnalyticsConfig.APIServerET = Config.APIServer;
 	AnalyticsConfig.AppVersionET = Config.ApiVersion;
+	AnalyticsConfig.AppEnvironment = TEXT("unknown");
+	AnalyticsConfig.UploadType = TEXT("qosmetrics");
 	AnalyticsProvider = FAnalyticsET::Get().CreateAnalyticsProvider(AnalyticsConfig);
 	// we'll just use the MachineID for the User. It actually won't matter much for this service.
 	AnalyticsProvider->SetUserID(FPlatformMisc::GetMachineId().ToString(EGuidFormats::Digits).ToLower());
diff --git a/Engine/Source/Runtime/Analytics/AnalyticsET/Public/AnalyticsET.h b/Engine/Source/Runtime/Analytics/AnalyticsET/Public/AnalyticsET.h
--- a/Engine/Source/Runtime/Analytics/AnalyticsET/Public/AnalyticsET.h
+++ b/Engine/Source/Runtime/Analytics/AnalyticsET/Public/AnalyticsET.h
@@ -38,7 +38,7 @@ public:
 	{
 		/** ET APIKey - Get from your account manager */
 		FString APIKeyET;
-		/** ET API Server - Defaults if empty to GetDefaultAPIServer. */
+		/** ET API Server - Base URL to send events. */
 		FString APIServerET;
 		/** 
 		 * AppVersion - defines the app version passed to the provider. By default this will be FEngineVersion::Current(), but you can supply your own. 
@@ -48,6 +48,10 @@ public:
 		FString AppVersionET;
 		/** When true, sends events using the legacy ET protocol that passes all attributes as URL parameters. Defaults to false. */
 		bool UseLegacyProtocol;
+		/** The AppEnvironment that the data router should use. Defaults to GetDefaultAppEnvironment. */
+		FString AppEnvironment;
+		/** The UploadType that the data router should use. Defaults to GetDefaultUploadType. */
+		FString UploadType;
 
 		/** Default ctor to ensure all values have their proper default. */
 		Config() : UseLegacyProtocol(false) {}
@@ -60,8 +64,14 @@ public:
 		static FString GetKeyNameForAppVersion() { return TEXT("AppVersionET"); }
 		/** Optional parameter to use the legacy backend protocol. */
 		static FString GetKeyNameForUseLegacyProtocol() { return TEXT("UseLegacyProtocol"); }
+		/** For the the data router backend protocol. */
+		static FString GetKeyNameForAppEnvironment() { return TEXT("AppEnvironment"); }
+		/** For the the data router backend protocol. */
+		static FString GetKeyNameForUploadType() { return TEXT("UploadType"); }
 		/** Default value if no APIServer configuration is provided. */
-		static FString GetDefaultAPIServer() { return TEXT("http://devonline-02:/ETAP/"); }
+		static FString GetDefaultAppEnvironment() { return TEXT("datacollector-binary"); }
+		/** Default value if no UploadType is given, and UseDataRouter protocol is specified. */
+		static FString GetDefaultUploadType() { return TEXT("eteventstream"); }
 	};
 
 	//--------------------------------------------------------------------------
diff --git a/Engine/Source/Runtime/AutomationWorker/Private/AutomationAnalytics.cpp b/Engine/Source/Runtime/AutomationWorker/Private/AutomationAnalytics.cpp
--- a/Engine/Source/Runtime/AutomationWorker/Private/AutomationAnalytics.cpp
+++ b/Engine/Source/Runtime/AutomationWorker/Private/AutomationAnalytics.cpp
@@ -5,6 +5,7 @@
 #include "EngineBuildSettings.h"
 #include "Runtime/Analytics/Analytics/Public/Analytics.h"
 #include "Runtime/Analytics/Analytics/Public/Interfaces/IAnalyticsProvider.h"
+#include "AnalyticsET.h"
 
 DEFINE_LOG_CATEGORY(LogAutomationAnalytics);
 
@@ -14,18 +15,6 @@ TArray<FString> FAutomationAnalytics::AutomationEventNames;
 TArray<FString> FAutomationAnalytics::AutomationParamNames;
 FString FAutomationAnalytics::MachineSpec;
 
-/**
-* Automation analytics config log to initialize the automation analytics provider.
-* External code should bind this delegate if automation analytics are desired,
-* preferably in private code that won't be redistributed.
-*/
-FAnalytics::FProviderConfigurationDelegate& GetAutomationAnalyticsOverrideConfigDelegate()
-{
-	static FAnalytics::FProviderConfigurationDelegate Delegate;
-	return Delegate;
-}
-
-
 /**
 * Get analytics pointer
 */
@@ -40,99 +29,22 @@ void FAutomationAnalytics::Initialize()
 {
 	checkf(!bIsInitialized, TEXT("FAutomationAnalytics::Initialize called more than once."));
 
-	// this will only be true for builds that have editor support (currently PC, Mac, Linux)
-	// The idea here is to only send editor events for actual editor runs, not for things like -game runs of the editor.
-	//const bool bIsEditorRun = WITH_EDITOR && GIsEditor && !IsRunningCommandlet();
-	const bool bIsEditorRun = GIsEditor;
-
-	// We also want to identify a real run of a game, which is NOT necessarily the opposite of an editor run.
-	// Ideally we'd be able to tell explicitly, but with content-only games, it becomes difficult.
-	// So we ensure we are not an editor run, we don't have EDITOR stuff compiled in, we are not running a commandlet,
-	// we are not a generic, utility program, and we require cooked data.
-	//const bool bIsGameRun = !WITH_EDITOR && !IsRunningCommandlet() && !FPlatformProperties::IsProgram() && FPlatformProperties::RequiresCookedData();
-
-	const bool bIsGameRun = !GIsEditor;
-
-	const bool bShouldInitAnalytics = bIsEditorRun || bIsGameRun;
-
-	if (bShouldInitAnalytics)
+	//Set the config for Analytics
+	FAnalyticsET::Config Config;
+	Config.APIKeyET = FString::Printf(TEXT("AutomationAnalytics.%s"), GInternalGameName);
+	Config.APIServerET = TEXT("https://datarouter.ol.epicgames.com/");
+	Analytics = FAnalyticsET::Get().CreateAnalyticsProvider(Config);
+	if (Analytics.IsValid())
 	{
-		// Setup some default automation analytics if there is nothing custom bound
-		FAnalytics::FProviderConfigurationDelegate DefaultAutomationAnalyticsConfig;
-		DefaultAutomationAnalyticsConfig.BindLambda(
-			[=](const FString& KeyName, bool bIsValueRequired) -> FString
-		{
-			static TMap<FString, FString> ConfigMap;
-			if (ConfigMap.Num() == 0)
-			{
-				ConfigMap.Add(TEXT("ProviderModuleName"), TEXT("AnalyticsET"));
-				ConfigMap.Add(TEXT("APIServerET"), TEXT("http://etsource.epicgames.com/ET2/"));
-
-				// We always use the "Release" analytics account unless we're running in analytics test mode (usually with
-				// a command-line parameter), or we're an internal Epic build
-				const FAnalytics::BuildType AnalyticsBuildType = FAnalytics::Get().GetBuildType();
-				const bool bUseReleaseAccount =
-					(AnalyticsBuildType == FAnalytics::Development || AnalyticsBuildType == FAnalytics::Release) &&
-					!FEngineBuildSettings::IsInternalBuild();   // Internal Epic build
-				const TCHAR* BuildTypeStr = bUseReleaseAccount ? TEXT("Release") : TEXT("Dev");
+		Analytics->SetUserID(FString::Printf(TEXT("%s|%s|%s"), *FPlatformMisc::GetMachineId().ToString(EGuidFormats::Digits).ToLower(), *FPlatformMisc::GetEpicAccountId(), *FPlatformMisc::GetOperatingSystemId()));
+		Analytics->StartSession();
 
-				FString UE4TypeOverride;
-				bool bHasOverride = FAnalytics::ConfigFromIni::GetUE4TypeOverride(UE4TypeOverride);
-				const TCHAR* UE4TypeStr = bHasOverride ? *UE4TypeOverride : FEngineBuildSettings::IsPerforceBuild() ? TEXT("Perforce") : TEXT("UnrealEngine");
-				if (bIsEditorRun)
-				{
-					ConfigMap.Add(TEXT("APIKeyET"), FString::Printf(TEXT("UEEditor.Automation.%s.%s"), UE4TypeStr, BuildTypeStr));
-				}
-				else
-				{
-					const TCHAR* GameName = GInternalGameName;
-					ConfigMap.Add(TEXT("APIKeyET"), FString::Printf(TEXT("UEGame.Automation.%s.%s.%s"), UE4TypeStr, BuildTypeStr, GameName));
-				}
-			}
+		MachineSpec = FParse::Param(FCommandLine::Get(), TEXT("60hzmin")) 
+			? TEXT("60hzmin")
+			: FParse::Param(FCommandLine::Get(), TEXT("30hzmin")) 
+			? TEXT("30hzmin") 
+			: TEXT("");
 
-			// Check for overrides
-			if (GetAutomationAnalyticsOverrideConfigDelegate().IsBound())
-			{
-				const FString OverrideValue = GetAutomationAnalyticsOverrideConfigDelegate().Execute(KeyName, bIsValueRequired);
-				if (!OverrideValue.IsEmpty())
-				{
-					return OverrideValue;
-				}
-			}
-
-			FString* ConfigValue = ConfigMap.Find(KeyName);
-			return ConfigValue != NULL ? *ConfigValue : TEXT("");
-		});
-
-		Analytics = FAnalytics::Get().CreateAnalyticsProvider(
-			FName(*DefaultAutomationAnalyticsConfig.Execute(TEXT("ProviderModuleName"), true)),
-			DefaultAutomationAnalyticsConfig);
-		if (Analytics.IsValid())
-		{
-			Analytics->SetUserID(FString::Printf(TEXT("%s|%s|%s"), *FPlatformMisc::GetMachineId().ToString(EGuidFormats::Digits).ToLower(), *FPlatformMisc::GetEpicAccountId(), *FPlatformMisc::GetOperatingSystemId()));
-
-			TArray<FAnalyticsEventAttribute> StartSessionAttributes;
-			// Add project info whether we are in editor or game.
-			FString GameName = GInternalGameName;
-			StartSessionAttributes.Emplace(TEXT("ProjectName"), GameName);
-
-
-			//Need a way to get these without using the Editor
-			/*StartSessionAttributes.Emplace(TEXT("ProjectID"), ProjectSettings.ProjectID);
-			StartSessionAttributes.Emplace(TEXT("ProjectDescription"), ProjectSettings.Description);
-			StartSessionAttributes.Emplace(TEXT("ProjectVersion"), ProjectSettings.ProjectVersion);*/
-
-			Analytics->StartSession(StartSessionAttributes);
-
-			MachineSpec = FParse::Param(FCommandLine::Get(), TEXT("60hzmin")) 
-				? TEXT("60hzmin")
-				: FParse::Param(FCommandLine::Get(), TEXT("30hzmin")) 
-				? TEXT("30hzmin") 
-				: TEXT("");
-		}
-	}
-	if (Analytics.IsValid())
-	{
 		bIsInitialized = true;
 		InititalizeAnalyticParameterNames();
 	}
diff --git a/Engine/Source/Runtime/Core/Private/CorePrivatePCH.h b/Engine/Source/Runtime/Core/Private/CorePrivatePCH.h
--- a/Engine/Source/Runtime/Core/Private/CorePrivatePCH.h
+++ b/Engine/Source/Runtime/Core/Private/CorePrivatePCH.h
@@ -78,7 +78,6 @@ Includes.
 #include "CircularBuffer.h"
 #include "Queue.h"
 #include "Ticker.h"						// Efficient scheduled delegate manager
-#include "RocketSupport.h"				// Core support for launching in "Rocket" mode
 #include "ProfilingHelpers.h"			// Misc profiling helpers.
 #include "ConfigCacheIni.h"				// The configuration cache declarations
 #include "IConsoleManager.h"
diff --git a/Engine/Source/Runtime/Core/Private/Misc/App.cpp b/Engine/Source/Runtime/Core/Private/Misc/App.cpp
--- a/Engine/Source/Runtime/Core/Private/Misc/App.cpp
+++ b/Engine/Source/Runtime/Core/Private/Misc/App.cpp
@@ -136,10 +136,19 @@ bool FApp::IsEngineInstalled()
 
 	if (EngineInstalledState == -1)
 	{
-		bool bIsInstalledEngine = IsInstalled() || (FRocketSupport::IsRocket() ? !FParse::Param(FCommandLine::Get(), TEXT("NotInstalledEngine")) : FParse::Param(FCommandLine::Get(), TEXT("InstalledEngine")));
+		bool bIsInstalledEngine = IsInstalled();
 		FString InstalledBuildFile = FPaths::RootDir() / TEXT("Engine/Build/InstalledBuild.txt");
 		FPaths::NormalizeFilename(InstalledBuildFile);
 		bIsInstalledEngine |= IFileManager::Get().FileExists(*InstalledBuildFile);
+		// Allow commandline options to disable/enable installed engine behavior
+		if (bIsInstalledEngine)
+		{
+			bIsInstalledEngine = !FParse::Param(FCommandLine::Get(), TEXT("NotInstalledEngine"));
+		}
+		else
+		{
+			bIsInstalledEngine = FParse::Param(FCommandLine::Get(), TEXT("InstalledEngine"));
+		}
 		EngineInstalledState = bIsInstalledEngine ? 1 : 0;
 	}
 
diff --git a/Engine/Source/Runtime/Core/Public/Core.h b/Engine/Source/Runtime/Core/Public/Core.h
--- a/Engine/Source/Runtime/Core/Public/Core.h
+++ b/Engine/Source/Runtime/Core/Public/Core.h
@@ -94,7 +94,6 @@
 #include "CircularQueue.h"
 #include "Queue.h"
 #include "Ticker.h"						// Efficient scheduled delegate manager
-#include "RocketSupport.h"				// Core support for launching in "Rocket" mode
 #include "ProfilingHelpers.h"			// Misc profiling helpers.
 #include "ConfigCacheIni.h"				// The configuration cache declarations
 #include "IConsoleManager.h"
diff --git a/Engine/Source/Runtime/Engine/Private/EngineAnalytics.cpp b/Engine/Source/Runtime/Engine/Private/EngineAnalytics.cpp
--- a/Engine/Source/Runtime/Engine/Private/EngineAnalytics.cpp
+++ b/Engine/Source/Runtime/Engine/Private/EngineAnalytics.cpp
@@ -65,7 +65,8 @@ void FEngineAnalytics::Initialize()
 					if (ConfigMap.Num() == 0)
 					{
 						ConfigMap.Add(TEXT("ProviderModuleName"), TEXT("AnalyticsET"));
-						ConfigMap.Add(TEXT("APIServerET"), TEXT("http://etsource.epicgames.com/ET2/"));
+						ConfigMap.Add(TEXT("APIServerET"), TEXT("https://datarouter.ol.epicgames.com/"));
+						ConfigMap.Add(TEXT("AppEnvironment"), TEXT("datacollector-source"));
 
 						// We always use the "Release" analytics account unless we're running in analytics test mode (usually with
 						// a command-line parameter), or we're an internal Epic build
@@ -76,7 +77,7 @@ void FEngineAnalytics::Initialize()
 						const TCHAR* BuildTypeStr = bUseReleaseAccount ? TEXT("Release") : TEXT("Dev");
 
 						FString UE4TypeOverride;
-						bool bHasOverride = FAnalytics::ConfigFromIni::GetUE4TypeOverride(UE4TypeOverride);
+						bool bHasOverride = GConfig->GetString(TEXT("Analytics"), TEXT("UE4TypeOverride"), UE4TypeOverride, GEngineIni);
 						const TCHAR* UE4TypeStr = bHasOverride ? *UE4TypeOverride : FEngineBuildSettings::IsPerforceBuild() ? TEXT("Perforce") : TEXT("UnrealEngine");
 						if (bIsEditorRun)
 						{
diff --git a/Engine/Source/Runtime/SQLiteSupport/SQLiteSupport.Build.cs b/Engine/Source/Runtime/SQLiteSupport/SQLiteSupport.Build.cs
--- a/Engine/Source/Runtime/SQLiteSupport/SQLiteSupport.Build.cs
+++ b/Engine/Source/Runtime/SQLiteSupport/SQLiteSupport.Build.cs
@@ -45,11 +45,6 @@ namespace UnrealBuildTool.Rules
 					break;
 			}
 
-			if (UnrealBuildTool.RunningRocket())
-			{
-				throw new BuildException("This module requires a source code build of the engine from Github. Please refer to the Engine/Source/ThirdParty/sqlite/README.txt file prior to enabling this module.");
-			}
-		
 			string LibraryPath = "" + UEBuildConfiguration.UEThirdPartySourceDirectory + "sqlite/lib/" + PlatformName + ConfigurationName;
 			string LibraryFilename = Path.Combine(LibraryPath, "sqlite" + UEBuildPlatform.GetBuildPlatform(Target.Platform).GetBinaryExtension(UEBuildBinaryType.StaticLibrary));
 			if (!File.Exists(LibraryFilename))
