Remove sending of UniqueAdvertisingId from AnalyticsET SessionStart events.

diff --git a/Engine/Source/Runtime/Analytics/AnalyticsET/Private/IAnalyticsProviderET.cpp b/Engine/Source/Runtime/Analytics/AnalyticsET/Private/IAnalyticsProviderET.cpp
--- a/Engine/Source/Runtime/Analytics/AnalyticsET/Private/IAnalyticsProviderET.cpp
+++ b/Engine/Source/Runtime/Analytics/AnalyticsET/Private/IAnalyticsProviderET.cpp
@@ -445,8 +445,6 @@ bool FAnalyticsProviderET::StartSession(TArray<FAnalyticsEventAttribute>&& Attri
 
 	// always ensure we send a few specific attributes on session start.
 	TArray<FAnalyticsEventAttribute> AppendedAttributes(MoveTemp(Attributes));
-	// this allows mapping to ad networks attribution data
-	AppendedAttributes.Emplace(TEXT("AttributionId"), FPlatformMisc::GetUniqueAdvertisingId());
 	// we should always know what platform is hosting this session.
 	AppendedAttributes.Emplace(TEXT("Platform"), FString(FPlatformProperties::IniPlatformName()));
 
