Expose some additional analytics values

diff --git a/Engine/Source/Runtime/Analytics/AnalyticsET/Private/IAnalyticsProviderET.cpp b/Engine/Source/Runtime/Analytics/AnalyticsET/Private/IAnalyticsProviderET.cpp
--- a/Engine/Source/Runtime/Analytics/AnalyticsET/Private/IAnalyticsProviderET.cpp
+++ b/Engine/Source/Runtime/Analytics/AnalyticsET/Private/IAnalyticsProviderET.cpp
@@ -188,6 +188,7 @@ public:
 
 	virtual void SetAppID(FString&& AppID) override;
 	virtual const FString& GetAppID() const override;
+	virtual const FString& GetAppVersion() const override;
 	virtual void SetUserID(const FString& InUserID) override;
 	virtual FString GetUserID() const override;
 
@@ -771,6 +772,11 @@ const FString& FAnalyticsProviderET::GetAppID() const
 	return APIKey;
 }
 
+const FString& FAnalyticsProviderET::GetAppVersion() const
+{
+	return AppVersion;
+}
+
 void FAnalyticsProviderET::SetUserID(const FString& InUserID)
 {
 	// command-line specified user ID overrides all attempts to reset it.
diff --git a/Engine/Source/Runtime/Analytics/AnalyticsET/Public/IAnalyticsProviderET.h b/Engine/Source/Runtime/Analytics/AnalyticsET/Public/IAnalyticsProviderET.h
--- a/Engine/Source/Runtime/Analytics/AnalyticsET/Public/IAnalyticsProviderET.h
+++ b/Engine/Source/Runtime/Analytics/AnalyticsET/Public/IAnalyticsProviderET.h
@@ -28,6 +28,13 @@ public:
 	 */
 	virtual const FString& GetAppID() const = 0;
 
+	/**
+	* Method to get the AppVersion
+	*
+	* @return the AppVersion
+	*/
+	virtual const FString& GetAppVersion() const = 0;
+
 	/**
 	* Optimization for StartSession that avoids the array copy using rvalue references.
 	*
