AnalyticsConversion methods moved from a namespace to global namespace to avoid 2-phase template instantiation issues with CLANG compilers that cause include order sensitivity with AnalyticsEventAttribute.h.

Also add an overload for FStrings to allow us to avoid copies in that specific case, since we can't use a universal qualifier in the primary template, which caused unnecessary string copies to be made in almost all the code that added attributes.

#rnx

diff --git a/Engine/Source/Runtime/Analytics/Analytics/Public/AnalyticsConversion.h b/Engine/Source/Runtime/Analytics/Analytics/Public/AnalyticsConversion.h
--- a/Engine/Source/Runtime/Analytics/Analytics/Public/AnalyticsConversion.h
+++ b/Engine/Source/Runtime/Analytics/Analytics/Public/AnalyticsConversion.h
@@ -5,59 +5,79 @@
 #include "CoreMinimal.h"
 
 /** Helpers for converting various common types to strings that analytics providers can consume. */
-namespace AnalyticsConversion
-{
-	/** Lexical conversion. Allow any type that we have a Lex for. Can't use universal references here because it then eats all non-perfect matches for the array and TMap conversions below, which we want to use a custom, analytics specific implementation for. */
-	template <typename T>
+
+/** Lexical conversion. Allow any type that we have a Lex for. Can't use universal references here because it then eats all non-perfect matches for the array and TMap conversions below, which we want to use a custom, analytics specific implementation for. */
+template <typename T>
 #if PLATFORM_COMPILER_HAS_DECLTYPE_AUTO
-	inline decltype(auto) ToString(const T& Value)
+inline decltype(auto) AnalyticsConversionToString(const T& Value)
 #else
-	inline auto ToString(const T& Value) -> decltype(LexToString(Value)) 
+inline auto AnalyticsConversionToString(const T& Value) -> decltype(LexToString(Value)) 
 #endif
-	{
-		return LexToString(Value);
-	}
-	inline FString ToString(float Value)
-	{
-		return LexToSanitizedString(Value);
-	}
-	inline FString ToString(double Value)
-	{
-		return LexToSanitizedString(Value);
-	}
+{
+	return LexToString(Value);
+}
+/** Make sure we have a direct implementation for moving FStrings as we definitely don't want to copy them as the above const-ref template will cause. */
+inline FString AnalyticsConversionToString(FString&& Value)
+{
+	return MoveTemp(Value);
+}
+inline FString AnalyticsConversionToString(const FString& Value)
+{
+	return Value;
+}
+inline FString AnalyticsConversionToString(float Value)
+{
+	return LexToSanitizedString(Value);
+}
+inline FString AnalyticsConversionToString(double Value)
+{
+	return LexToSanitizedString(Value);
+}
 
-	/** Array conversion. Creates comma-separated list. */
-	template <typename T, typename AllocatorType>
-	FString ToString(const TArray<T, AllocatorType>& ValueArray)
+/** Array conversion. Creates comma-separated list. */
+template <typename T, typename AllocatorType>
+FString AnalyticsConversionToString(const TArray<T, AllocatorType>& ValueArray)
+{
+	FString Result;
+	// Serialize the array into "value1,value2,..." format
+	for (const T& Value : ValueArray)
 	{
-		FString Result;
-		// Serialize the array into "value1,value2,..." format
-		for (const T& Value : ValueArray)
-		{
-			Result += AnalyticsConversion::ToString(Value);
-			Result += TEXT(",");
-		}
-		// Remove the trailing comma (LeftChop will ensure an empty container won't crash here).
-		Result = Result.LeftChop(1);
-		return Result;
+		Result += AnalyticsConversionToString(Value);
+		Result += TEXT(",");
 	}
+	// Remove the trailing comma (LeftChop will ensure an empty container won't crash here).
+	Result = Result.LeftChop(1);
+	return Result;
+}
 
-	/** Map conversion. Creates comma-separated list. Creates comma-separated list with colon-separated key:value pairs. */
-	template<typename KeyType, typename ValueType, typename Allocator, typename KeyFuncs>
-	FString ToString(const TMap<KeyType, ValueType, Allocator, KeyFuncs>& ValueMap)
+/** Map conversion. Creates comma-separated list. Creates comma-separated list with colon-separated key:value pairs. */
+template<typename KeyType, typename ValueType, typename Allocator, typename KeyFuncs>
+FString AnalyticsConversionToString(const TMap<KeyType, ValueType, Allocator, KeyFuncs>& ValueMap)
+{
+	FString Result;
+	// Serialize the map into "key:value,..." format
+	for (auto& KVP : ValueMap)
 	{
-		FString Result;
-		// Serialize the map into "key:value,..." format
-		for (auto& KVP : ValueMap)
-		{
-			Result += AnalyticsConversion::ToString(KVP.Key);
-			Result += TEXT(":");
-			Result += AnalyticsConversion::ToString(KVP.Value);
-			Result += TEXT(",");
-		}
-		// Remove the trailing comma (LeftChop will ensure an empty container won't crash here).
-		Result = Result.LeftChop(1);
-		return Result;
+		Result += AnalyticsConversionToString(KVP.Key);
+		Result += TEXT(":");
+		Result += AnalyticsConversionToString(KVP.Value);
+		Result += TEXT(",");
 	}
+	// Remove the trailing comma (LeftChop will ensure an empty container won't crash here).
+	Result = Result.LeftChop(1);
+	return Result;
 }
 
+/** Deprecated namespace. Forwards on to global equivalents for backwards compatibility. See functions above. Namespace was deprecated for the same reason as Lex::ToString -> LexToString. See UnrealString.h for details. */
+namespace AnalyticsConversion
+{
+	template <typename T>
+#if PLATFORM_COMPILER_HAS_DECLTYPE_AUTO
+	inline decltype(auto) ToString(T&& Value)
+#else
+	inline auto ToString(T&& Value) -> decltype(AnalyticsConversionToString(Forward<T>(Value))) 
+#endif
+	{
+		return AnalyticsConversionToString(Forward<T>(Value));
+	}
+}
\ No newline at end of file
diff --git a/Engine/Source/Runtime/Analytics/Analytics/Public/AnalyticsEventAttribute.h b/Engine/Source/Runtime/Analytics/Analytics/Public/AnalyticsEventAttribute.h
--- a/Engine/Source/Runtime/Analytics/Analytics/Public/AnalyticsEventAttribute.h
+++ b/Engine/Source/Runtime/Analytics/Analytics/Public/AnalyticsEventAttribute.h
@@ -142,15 +142,15 @@ public: // json fragment
 
 public: // string (catch-all)
 	/**
-	 * Helper constructor to make an attribute from a name/value pair by forwarding through LexToString and AnalyticsConversion::ToString.
+	 * Helper constructor to make an attribute from a name/value pair by forwarding through LexToString and AnalyticsConversionToString.
 	 * 
 	 * @param InName Name of the attribute. Will be converted to a string via forwarding to LexToString
-	 * @param InValue Value of the attribute. Will be converted to a string via forwarding to AnalyticsConversion::ToString (same as Lex but with basic support for arrays and maps)
+	 * @param InValue Value of the attribute. Will be converted to a string via forwarding to AnalyticsConversionToString (same as Lex but with basic support for arrays and maps)
 	 */
 	template <typename NameType, typename ValueType>
 	FAnalyticsEventAttribute(NameType&& InName, ValueType&& InValue)
 		: AttrName(LexToString(Forward<NameType>(InName)))
-		, AttrValueString(AnalyticsConversion::ToString(Forward<ValueType>(InValue)))
+		, AttrValueString(AnalyticsConversionToString(Forward<ValueType>(InValue)))
 		, AttrValueNumber(0)
 		, AttrValueBool(false)
 		, AttrType(AttrTypeEnum::String)
