Add copy ctors and assignment operators to FAnalyticsEventAttribute. That class has a hacky "immutable" state by using const members, which suppresses assignability, which was causing people to go to great contortions in various spots in the code to assign over the top of existing members. As a workaround, I implemented these operators that const_cast<> the members to let us assign over them, but still not mutate them directly.
#tests This is split from another larger change to keep the engine changes separate.

diff --git a/Engine/Source/Runtime/Analytics/Analytics/Public/AnalyticsEventAttribute.h b/Engine/Source/Runtime/Analytics/Analytics/Public/AnalyticsEventAttribute.h
--- a/Engine/Source/Runtime/Analytics/Analytics/Public/AnalyticsEventAttribute.h
+++ b/Engine/Source/Runtime/Analytics/Analytics/Public/AnalyticsEventAttribute.h
@@ -47,6 +47,53 @@ struct FAnalyticsEventAttribute
 		, AttrType(AttrTypeEnum::String)
 	{}
 
+	/** Reinstate the default copy ctor because that one still works fine. */
+	FAnalyticsEventAttribute(const FAnalyticsEventAttribute& RHS) = default;
+
+	/** Hack to allow copy ctor using an rvalue-ref. This class only "sort of" acts like an immutable class because the const members prevents assignment, which was not intended when this code was changed. */
+	FAnalyticsEventAttribute(FAnalyticsEventAttribute&& RHS)
+		: AttrName(MoveTemp(const_cast<FString&>(RHS.AttrName)))
+		, AttrValueString(MoveTemp(const_cast<FString&>(RHS.AttrValueString)))
+		// no need to use MoveTemp on intrinsic types.
+		, AttrValueNumber(RHS.AttrValueNumber)
+		, AttrValueBool(RHS.AttrValueBool)
+		, AttrType(RHS.AttrType)
+	{
+	}
+
+	/** Hack to allow assignment. This class only "sort of" acts like an immutable class because the const members prevents assignment, which was not intended when this code was changed. */
+	FAnalyticsEventAttribute& operator=(const FAnalyticsEventAttribute& RHS)
+	{
+		if (&RHS == this)
+		{
+			return *this;
+		}
+
+		const_cast<FString&>(AttrName) = RHS.AttrName;
+		const_cast<FString&>(AttrValueString) = RHS.AttrValueString;
+		const_cast<double&>(AttrValueNumber) = RHS.AttrValueNumber;
+		const_cast<bool&>(AttrValueBool) = RHS.AttrValueBool;
+		const_cast<AttrTypeEnum&>(AttrType) = RHS.AttrType;
+		return *this;
+	}
+
+	/** Hack to allow assignment. This class only "sort of" acts like an immutable class because the const members prevents assignment, which was not intended when this code was changed. */
+	FAnalyticsEventAttribute& operator=(FAnalyticsEventAttribute&& RHS)
+	{
+		if (&RHS == this)
+		{
+			return *this;
+		}
+
+		const_cast<FString&>(AttrName) = MoveTemp(const_cast<FString&>(RHS.AttrName));
+		const_cast<FString&>(AttrValueString) = MoveTemp(const_cast<FString&>(RHS.AttrValueString));
+		// no need to use MoveTemp on intrinsic types.
+		const_cast<double&>(AttrValueNumber) = RHS.AttrValueNumber;
+		const_cast<bool&>(AttrValueBool) = RHS.AttrValueBool;
+		const_cast<AttrTypeEnum&>(AttrType) = RHS.AttrType;
+		return *this;
+	}
+
 	/** If you need the old AttrValue behavior (i.e. stringify everything), call this function instead. */
 	FString ToString() const
 	{
