- Add hotfix version to all server events.
- Replace format of app version sent in ET analytics to: Season.Minor.Hotfix - %VERSION%

diff --git a/Engine/Source/Runtime/Analytics/AnalyticsET/Private/IAnalyticsProviderET.cpp b/Engine/Source/Runtime/Analytics/AnalyticsET/Private/IAnalyticsProviderET.cpp
--- a/Engine/Source/Runtime/Analytics/AnalyticsET/Private/IAnalyticsProviderET.cpp
+++ b/Engine/Source/Runtime/Analytics/AnalyticsET/Private/IAnalyticsProviderET.cpp
@@ -188,6 +188,7 @@ public:
 
 	virtual void SetAppID(FString&& AppID) override;
 	virtual const FString& GetAppID() const override;
+	virtual void SetAppVersion(FString&& AppVersion) override;
 	virtual const FString& GetAppVersion() const override;
 	virtual void SetUserID(const FString& InUserID) override;
 	virtual FString GetUserID() const override;
@@ -767,6 +768,16 @@ void FAnalyticsProviderET::SetAppID(FString&& InAppID)
 	}
 }
 
+void FAnalyticsProviderET::SetAppVersion(FString&& InAppVersion)
+{
+	if (AppVersion != InAppVersion)
+	{
+		// Flush any cached events that would be using the old AppVersion.
+		FlushEvents();
+		AppVersion = MoveTemp(InAppVersion);
+	}
+}
+
 const FString& FAnalyticsProviderET::GetAppID() const
 {
 	return APIKey;
diff --git a/Engine/Source/Runtime/Analytics/AnalyticsET/Public/IAnalyticsProviderET.h b/Engine/Source/Runtime/Analytics/AnalyticsET/Public/IAnalyticsProviderET.h
--- a/Engine/Source/Runtime/Analytics/AnalyticsET/Public/IAnalyticsProviderET.h
+++ b/Engine/Source/Runtime/Analytics/AnalyticsET/Public/IAnalyticsProviderET.h
@@ -28,6 +28,13 @@ public:
 	 */
 	virtual const FString& GetAppID() const = 0;
 
+	/**
+	 * Sets the AppVersion.
+	 *
+	 * @param AppVersion The new AppVersion.
+	 */
+	virtual void SetAppVersion(FString&& AppVersion) = 0;
+
 	/**
 	* Method to get the AppVersion
 	*
