Fix FAnalyticsProviderET::SetAppVersion to apply the %VERSION% replacement before it compares the new/old value
#tests run a local server, wait for hotfix, verify that version replacement works and no version change is made.
#rnx

diff --git a/Engine/Source/Runtime/Analytics/AnalyticsET/Private/IAnalyticsProviderET.cpp b/Engine/Source/Runtime/Analytics/AnalyticsET/Private/IAnalyticsProviderET.cpp
--- a/Engine/Source/Runtime/Analytics/AnalyticsET/Private/IAnalyticsProviderET.cpp
+++ b/Engine/Source/Runtime/Analytics/AnalyticsET/Private/IAnalyticsProviderET.cpp
@@ -770,8 +770,14 @@ void FAnalyticsProviderET::SetAppID(FString&& InAppID)
 
 void FAnalyticsProviderET::SetAppVersion(FString&& InAppVersion)
 {
+	// make sure to do the version replacement if the given string is parameterized.
+	InAppVersion = InAppVersion.IsEmpty()
+		? FString(FApp::GetBuildVersion())
+		: InAppVersion.Replace(TEXT("%VERSION%"), FApp::GetBuildVersion(), ESearchCase::CaseSensitive);
+
 	if (AppVersion != InAppVersion)
 	{
+		UE_LOG(LogAnalytics, Log, TEXT("[%s] Updating AppVersion to %s from old value of %s"), *APIKey, *InAppVersion, *AppVersion);
 		// Flush any cached events that would be using the old AppVersion.
 		FlushEvents();
 		AppVersion = MoveTemp(InAppVersion);
