Added timing scopes for FAnalyticsProviderET::FlushEvents and FCurlHttpRequest Setup/Process.

#rnx

diff --git a/Engine/Source/Runtime/Analytics/AnalyticsET/Private/IAnalyticsProviderET.cpp b/Engine/Source/Runtime/Analytics/AnalyticsET/Private/IAnalyticsProviderET.cpp
--- a/Engine/Source/Runtime/Analytics/AnalyticsET/Private/IAnalyticsProviderET.cpp
+++ b/Engine/Source/Runtime/Analytics/AnalyticsET/Private/IAnalyticsProviderET.cpp
@@ -517,6 +517,7 @@ TSharedRef<IHttpRequest> FAnalyticsProviderET::CreateRequest()
 
 void FAnalyticsProviderET::FlushEvents()
 {
+	QUICK_SCOPE_CYCLE_COUNTER(STAT_FAnalyticsProviderET_FlushEvents);
 	// Warn if this takes more than 2 ms
 	SCOPE_TIME_GUARD_MS(TEXT("FAnalyticsProviderET::FlushEvents"), 2);
 
@@ -552,56 +553,37 @@ void FAnalyticsProviderET::FlushEvents()
 				FlushedEvents->Reserve(CachedEvents.Num());
 			}
 
-			TSharedRef< TJsonWriter<TCHAR, TCondensedJsonPrintPolicy<TCHAR> > > JsonWriter = TJsonWriterFactory<TCHAR, TCondensedJsonPrintPolicy<TCHAR> >::Create(&Payload);
-			JsonWriter->WriteObjectStart();
-			JsonWriter->WriteArrayStart(TEXT("Events"));
-			for (FAnalyticsEventEntry& Entry : CachedEvents)
 			{
-				if (Entry.bIsDefaultAttributes)
+				QUICK_SCOPE_CYCLE_COUNTER(STAT_FlushEventsJson);
+				TSharedRef< TJsonWriter<TCHAR, TCondensedJsonPrintPolicy<TCHAR> > > JsonWriter = TJsonWriterFactory<TCHAR, TCondensedJsonPrintPolicy<TCHAR> >::Create(&Payload);
+				JsonWriter->WriteObjectStart();
+				JsonWriter->WriteArrayStart(TEXT("Events"));
+				for (FAnalyticsEventEntry& Entry : CachedEvents)
 				{
-					// This is the default attributes, so update the array.
-					if (FlushedEvents.IsValid())
+					if (Entry.bIsDefaultAttributes)
 					{
-						CurrentDefaultAttributes = Entry.Attributes; // need to copy
-						FlushedEvents->Emplace(MoveTemp(Entry));
-					}
-					else
-					{
-						CurrentDefaultAttributes = MoveTemp(Entry.Attributes);
-					}
-				}
-				else
-				{
-					++EventCount;
-
-					// event entry
-					JsonWriter->WriteObjectStart();
-					JsonWriter->WriteValue(TEXT("EventName"), Entry.EventName);
-					FString DateOffset = (CurrentTime - Entry.TimeStamp).ToString();
-					JsonWriter->WriteValue(TEXT("DateOffset"), DateOffset);
-					// default attributes for this event
-					for (const FAnalyticsEventAttribute& Attr : CurrentDefaultAttributes)
-					{
-						switch (Attr.AttrType)
+						// This is the default attributes, so update the array.
+						if (FlushedEvents.IsValid())
+						{
+							CurrentDefaultAttributes = Entry.Attributes; // need to copy
+							FlushedEvents->Emplace(MoveTemp(Entry));
+						}
+						else
 						{
-						case FAnalyticsEventAttribute::AttrTypeEnum::String:
-							JsonWriter->WriteValue(Attr.AttrName, Attr.AttrValueString);
-							break;
-						case FAnalyticsEventAttribute::AttrTypeEnum::Number:
-							JsonWriter->WriteValue(Attr.AttrName, Attr.ToString());
-							break;
-						case FAnalyticsEventAttribute::AttrTypeEnum::Boolean:
-							JsonWriter->WriteValue(Attr.AttrName, Attr.ToString());
-							break;
-						case FAnalyticsEventAttribute::AttrTypeEnum::JsonFragment:
-							JsonWriter->WriteRawJSONValue(Attr.AttrName, Attr.AttrValueString);
-							break;
+							CurrentDefaultAttributes = MoveTemp(Entry.Attributes);
 						}
 					}
-					// optional attributes for this event
-					if (!Entry.bIsJsonEvent)
+					else
 					{
-						for (const FAnalyticsEventAttribute& Attr : Entry.Attributes)
+						++EventCount;
+
+						// event entry
+						JsonWriter->WriteObjectStart();
+						JsonWriter->WriteValue(TEXT("EventName"), Entry.EventName);
+						FString DateOffset = (CurrentTime - Entry.TimeStamp).ToString();
+						JsonWriter->WriteValue(TEXT("DateOffset"), DateOffset);
+						// default attributes for this event
+						for (const FAnalyticsEventAttribute& Attr : CurrentDefaultAttributes)
 						{
 							switch (Attr.AttrType)
 							{
@@ -617,29 +599,51 @@ void FAnalyticsProviderET::FlushEvents()
 							case FAnalyticsEventAttribute::AttrTypeEnum::JsonFragment:
 								JsonWriter->WriteRawJSONValue(Attr.AttrName, Attr.AttrValueString);
 								break;
+							}
 						}
-					}
-					}
-					else
-					{
-						for (const FAnalyticsEventAttribute& Attr : Entry.Attributes)
+						// optional attributes for this event
+						if (!Entry.bIsJsonEvent)
 						{
-							JsonWriter->WriteRawJSONValue(Attr.AttrName, Attr.AttrValueString);
+							for (const FAnalyticsEventAttribute& Attr : Entry.Attributes)
+							{
+								switch (Attr.AttrType)
+								{
+								case FAnalyticsEventAttribute::AttrTypeEnum::String:
+									JsonWriter->WriteValue(Attr.AttrName, Attr.AttrValueString);
+									break;
+								case FAnalyticsEventAttribute::AttrTypeEnum::Number:
+									JsonWriter->WriteValue(Attr.AttrName, Attr.ToString());
+									break;
+								case FAnalyticsEventAttribute::AttrTypeEnum::Boolean:
+									JsonWriter->WriteValue(Attr.AttrName, Attr.ToString());
+									break;
+								case FAnalyticsEventAttribute::AttrTypeEnum::JsonFragment:
+									JsonWriter->WriteRawJSONValue(Attr.AttrName, Attr.AttrValueString);
+									break;
+								}
+							}
 						}
-					}
-					JsonWriter->WriteObjectEnd();
+						else
+						{
+							for (const FAnalyticsEventAttribute& Attr : Entry.Attributes)
+							{
+								JsonWriter->WriteRawJSONValue(Attr.AttrName, Attr.AttrValueString);
+							}
+						}
+						JsonWriter->WriteObjectEnd();
 
-					// move the entry into the flushed
-					if (FlushedEvents.IsValid())
-					{
-						FlushedEvents->Emplace(MoveTemp(Entry));
+						// move the entry into the flushed
+						if (FlushedEvents.IsValid())
+						{
+							FlushedEvents->Emplace(MoveTemp(Entry));
+						}
 					}
 				}
-			}
 
-			JsonWriter->WriteArrayEnd();
-			JsonWriter->WriteObjectEnd();
-			JsonWriter->Close();
+				JsonWriter->WriteArrayEnd();
+				JsonWriter->WriteObjectEnd();
+				JsonWriter->Close();
+			}
 
 			// UrlEncode NOTE: need to concatenate everything
 			FString URLPath  = TEXT("datarouter/api/v1/public/data?SessionID=") + FPlatformHttp::UrlEncode(SessionID);
@@ -666,22 +670,26 @@ void FAnalyticsProviderET::FlushEvents()
 				UE_LOG(LogAnalytics, VeryVerbose, TEXT("%s"), *LogString);
 			}
 
-			// Create/send Http request for an event
-			TSharedRef<IHttpRequest> HttpRequest = CreateRequest();
-			HttpRequest->SetHeader(TEXT("Content-Type"), TEXT("application/json; charset=utf-8"));
-			HttpRequest->SetURL(APIServer / URLPath);
-			HttpRequest->SetVerb(TEXT("POST"));
-			HttpRequest->SetContentAsString(Payload);
-
-			// Don't set a response callback if we are in our destructor, as the instance will no longer be there to call.
-			if (!bInDestructor)
 			{
-				HttpRequest->OnProcessRequestComplete().BindSP(this, &FAnalyticsProviderET::EventRequestComplete, FlushedEvents);
+				QUICK_SCOPE_CYCLE_COUNTER(STAT_FlushEventsHttpRequest);
+				// Create/send Http request for an event
+				TSharedRef<IHttpRequest> HttpRequest = CreateRequest();
+				HttpRequest->SetHeader(TEXT("Content-Type"), TEXT("application/json; charset=utf-8"));
+				HttpRequest->SetURL(APIServer / URLPath);
+				HttpRequest->SetVerb(TEXT("POST"));
+				HttpRequest->SetContentAsString(Payload);
+
+				// Don't set a response callback if we are in our destructor, as the instance will no longer be there to call.
+				if (!bInDestructor)
+				{
+					HttpRequest->OnProcessRequestComplete().BindSP(this, &FAnalyticsProviderET::EventRequestComplete, FlushedEvents);
+				}
+				HttpRequest->ProcessRequest();
 			}
-			HttpRequest->ProcessRequest();
 		}
 		else
 		{
+			QUICK_SCOPE_CYCLE_COUNTER(STAT_FlushEventsLegacy);
 			// this is a legacy pathway that doesn't accept batch payloads of cached data. We'll just send one request for each event, which will be slow for a large batch of requests at once.
 			for (auto& Event : CachedEvents)
 			{
diff --git a/Engine/Source/Runtime/Online/HTTP/Private/Curl/CurlHttp.cpp b/Engine/Source/Runtime/Online/HTTP/Private/Curl/CurlHttp.cpp
--- a/Engine/Source/Runtime/Online/HTTP/Private/Curl/CurlHttp.cpp
+++ b/Engine/Source/Runtime/Online/HTTP/Private/Curl/CurlHttp.cpp
@@ -694,6 +694,7 @@ size_t FCurlHttpRequest::DebugCallback(CURL * Handle, curl_infotype DebugInfoTyp
 
 bool FCurlHttpRequest::SetupRequest()
 {
+	QUICK_SCOPE_CYCLE_COUNTER(STAT_FCurlHttpRequest_SetupRequest);
 	check(EasyHandle);
 
 	if (!RequestPayload.IsValid())
@@ -894,6 +895,7 @@ bool FCurlHttpRequest::SetupRequest()
 
 bool FCurlHttpRequest::ProcessRequest()
 {
+	QUICK_SCOPE_CYCLE_COUNTER(STAT_FCurlHttpRequest_ProcessRequest);
 	check(EasyHandle);
 
 	// Clear the info cache log so we don't output messages from previous requests when reusing/retrying a request
@@ -925,6 +927,7 @@ bool FCurlHttpRequest::ProcessRequest()
 	}
 	else
 	{
+		QUICK_SCOPE_CYCLE_COUNTER(STAT_CurlHttpAddThreadedRequest);
 		// Mark as in-flight to prevent overlapped requests using the same object
 		CompletionStatus = EHttpRequestStatus::Processing;
 		// Response object to handle data that comes back after starting this request
