- Added AppId, AppVersion and UserId as FEditorAnalyticSession members and updated FEditorSessionSummaryWriter to write them.
  - Changed FEngineAnalytics::GetProvider() to return IAnalyticsProviderET rather than IAnalyticsProvider.
  - Updated FEditorSessionSummarySender to embed the AppId/AppVersion/UserId as attribute when sending the summary event and also configured the AnalyticsProvider to use the SessionID/AppID/AppVersion/UserID of the event rather than the current one.
  - Added a function to return the current IAnalyticProviderET configuration.

Edigrated 11445334 and 11459406.

diff --git a/Engine/Source/Developer/EditorAnalyticsSession/EditorAnalyticsSession.Build.cs b/Engine/Source/Developer/EditorAnalyticsSession/EditorAnalyticsSession.Build.cs
--- a/Engine/Source/Developer/EditorAnalyticsSession/EditorAnalyticsSession.Build.cs
+++ b/Engine/Source/Developer/EditorAnalyticsSession/EditorAnalyticsSession.Build.cs
@@ -10,7 +10,8 @@ namespace UnrealBuildTool.Rules
 				new string[]
 				{
 					"Core",
-					"Analytics"
+					"Analytics",
+					"AnalyticsET"
 				}
 			);
 		}
diff --git a/Engine/Source/Developer/EditorAnalyticsSession/Private/EditorAnalyticsSession.cpp b/Engine/Source/Developer/EditorAnalyticsSession/Private/EditorAnalyticsSession.cpp
--- a/Engine/Source/Developer/EditorAnalyticsSession/Private/EditorAnalyticsSession.cpp
+++ b/Engine/Source/Developer/EditorAnalyticsSession/Private/EditorAnalyticsSession.cpp
@@ -13,12 +13,21 @@ namespace EditorAnalyticsDefs
 	static const FString DefaultUserActivity(TEXT("Unknown"));
 	static const FString UnknownProjectValueString(TEXT("UnknownProject"));
 
+	static const FString UnknownAppIdString(TEXT("UnknownAppId"));
+	static const FString UnknownAppVersionString(TEXT("UnknownAppVersion"));
+	static const FString UnknownUserIdString(TEXT("UnknownUserID"));
+
 	// storage location
 	static const FString StoreId(TEXT("Epic Games"));
 	static const FString SessionSummarySection(TEXT("Unreal Engine/Session Summary/1_0"));
 	static const FString GlobalLockName(TEXT("UE4_SessionSummary_Lock"));
 	static const FString SessionListStoreKey(TEXT("SessionList"));
 
+	// capture context
+	static const FString AppIdStoreKey(TEXT("AppId"));
+	static const FString AppVersionStoreKey(TEXT("AppVersion"));
+	static const FString UserIdStoreKey(TEXT("UserId"));
+
 	// general values
 	static const FString ProjectNameStoreKey(TEXT("ProjectName"));
 	static const FString ProjectIDStoreKey(TEXT("ProjectID"));
@@ -125,6 +134,10 @@ namespace EditorAnalyticsUtils
 
 		FString SectionName = EditorAnalyticsUtils::GetSessionStorageLocation(Session.SessionId);
 
+		GET_STORED_STRING(AppId);
+		GET_STORED_STRING(AppVersion);
+		GET_STORED_STRING(UserId);
+
 		GET_STORED_STRING(ProjectName);
 		GET_STORED_STRING(ProjectID);
 		GET_STORED_STRING(ProjectDescription);
@@ -223,6 +236,10 @@ FSystemWideCriticalSection* FEditorAnalyticsSession::StoredValuesLock = nullptr;
 
 FEditorAnalyticsSession::FEditorAnalyticsSession()
 {
+	AppId = EditorAnalyticsDefs::UnknownAppIdString;
+	AppVersion = EditorAnalyticsDefs::UnknownAppVersionString;
+	UserId = EditorAnalyticsDefs::UnknownUserIdString;
+
 	ProjectName = EditorAnalyticsDefs::UnknownProjectValueString;
 	PlatformProcessID = 0;
 	StartupTimestamp = FDateTime::MinValue();
@@ -334,6 +351,10 @@ bool FEditorAnalyticsSession::Save()
 		const FString PluginsString = FString::Join(Plugins, TEXT(","));
 		FPlatformMisc::SetStoredValue(EditorAnalyticsDefs::StoreId, StorageLocation, EditorAnalyticsDefs::PluginsStoreKey, PluginsString);
 
+		SET_STORED_STRING(AppId);
+		SET_STORED_STRING(AppVersion);
+		SET_STORED_STRING(UserId);
+	
 		bAlreadySaved = true;
 	}
 
@@ -414,6 +435,10 @@ bool FEditorAnalyticsSession::Delete() const
 
 	FString SectionName = EditorAnalyticsUtils::GetSessionStorageLocation(SessionId);
 
+	FPlatformMisc::DeleteStoredValue(EditorAnalyticsDefs::StoreId, SectionName, EditorAnalyticsDefs::AppIdStoreKey);
+	FPlatformMisc::DeleteStoredValue(EditorAnalyticsDefs::StoreId, SectionName, EditorAnalyticsDefs::AppVersionStoreKey);
+	FPlatformMisc::DeleteStoredValue(EditorAnalyticsDefs::StoreId, SectionName, EditorAnalyticsDefs::UserIdStoreKey);
+
 	FPlatformMisc::DeleteStoredValue(EditorAnalyticsDefs::StoreId, SectionName, EditorAnalyticsDefs::ProjectNameStoreKey);
 	FPlatformMisc::DeleteStoredValue(EditorAnalyticsDefs::StoreId, SectionName, EditorAnalyticsDefs::ProjectIDStoreKey);
 	FPlatformMisc::DeleteStoredValue(EditorAnalyticsDefs::StoreId, SectionName, EditorAnalyticsDefs::ProjectDescriptionStoreKey);
diff --git a/Engine/Source/Developer/EditorAnalyticsSession/Private/EditorSessionSummarySender.cpp b/Engine/Source/Developer/EditorAnalyticsSession/Private/EditorSessionSummarySender.cpp
--- a/Engine/Source/Developer/EditorAnalyticsSession/Private/EditorSessionSummarySender.cpp
+++ b/Engine/Source/Developer/EditorAnalyticsSession/Private/EditorSessionSummarySender.cpp
@@ -4,9 +4,10 @@
 
 #include "AnalyticsEventAttribute.h"
 #include "Algo/Transform.h"
-#include "Interfaces/IAnalyticsProvider.h"
+#include "IAnalyticsProviderET.h"
 #include "EditorAnalyticsSession.h"
 #include "HAL/PlatformProcess.h"
+#include "Misc/EngineVersion.h"
 
 DEFINE_LOG_CATEGORY_STATIC(LogEditorSessionSummary, Verbose, All);
 
@@ -26,7 +27,7 @@ namespace EditorSessionSenderDefs
 	static const FString AbnormalSessionToken(TEXT("AbnormalShutdown"));
 }
 
-FEditorSessionSummarySender::FEditorSessionSummarySender(IAnalyticsProvider& InAnalyticsProvider, const FString& InSenderName, const int32 InCurrentSessionProcessId)
+FEditorSessionSummarySender::FEditorSessionSummarySender(IAnalyticsProviderET& InAnalyticsProvider, const FString& InSenderName, const int32 InCurrentSessionProcessId)
 	: HeartbeatTimeElapsed(0.0f)
 	, AnalyticsProvider(InAnalyticsProvider)
 	, Sender(InSenderName)
@@ -167,13 +168,21 @@ void FEditorSessionSummarySender::SendSessionSummaryEvent(const FEditorAnalytics
 	FString PluginsString = FString::Join(Session.Plugins, TEXT(","));
 
 	TArray<FAnalyticsEventAttribute> AnalyticsAttributes;
+
+	// Track which app/user is sending the summary (summary can be sent by another process (CrashReportClient) or another instance in case of crash/abnormal terminaison.
+	AnalyticsAttributes.Emplace(TEXT("SummarySenderAppId"), AnalyticsProvider.GetAppID());
+	AnalyticsAttributes.Emplace(TEXT("SummarySenderAppVersion"), AnalyticsProvider.GetAppVersion());
+	AnalyticsAttributes.Emplace(TEXT("SummarySenderEngineVersion"), FEngineVersion::Current().ToString(EVersionComponent::Changelist)); // Same as in EditorSessionSummaryWriter.cpp
+	AnalyticsAttributes.Emplace(TEXT("SummarySenderUserId"), AnalyticsProvider.GetUserID());
+	AnalyticsAttributes.Emplace(TEXT("SummarySenderSessionId"), AnalyticsProvider.GetSessionID()); // Not stripping the {} around the GUID like EditorSessionSummaryWriter does with SessionId.
+
 	AnalyticsAttributes.Emplace(TEXT("ProjectName"), Session.ProjectName);
 	AnalyticsAttributes.Emplace(TEXT("ProjectID"), Session.ProjectID);
 	AnalyticsAttributes.Emplace(TEXT("ProjectDescription"), Session.ProjectDescription);
 	AnalyticsAttributes.Emplace(TEXT("ProjectVersion"), Session.ProjectVersion);
 	AnalyticsAttributes.Emplace(TEXT("Platform"), FPlatformProperties::PlatformName());
-	AnalyticsAttributes.Emplace(TEXT("SessionId"), SessionIdString);
-	AnalyticsAttributes.Emplace(TEXT("EngineVersion"), Session.EngineVersion);
+	AnalyticsAttributes.Emplace(TEXT("SessionId"), SessionIdString); // The provider is expected to add it as "SessionID" param in the HTTP request, but keep it for completness, because the formats are slightly different.
+	AnalyticsAttributes.Emplace(TEXT("EngineVersion"), Session.EngineVersion); // The provider is expected to add it as "AppVersion" param in the HTTP request, but keep it for completness, because the formats are slightly different.
 	AnalyticsAttributes.Emplace(TEXT("ShutdownType"), ShutdownTypeString);
 	AnalyticsAttributes.Emplace(TEXT("StartupTimestamp"), Session.StartupTimestamp.ToIso8601());
 	AnalyticsAttributes.Emplace(TEXT("Timestamp"), Session.Timestamp.ToIso8601());
@@ -218,8 +227,31 @@ void FEditorSessionSummarySender::SendSessionSummaryEvent(const FEditorAnalytics
 		AnalyticsAttributes.Emplace(TEXT("ExitCode"), CurrentSessionExitCode.GetValue());
 	}
 
-	// Send the event.
-	AnalyticsProvider.RecordEvent(TEXT("SessionSummary"), AnalyticsAttributes);
+	// Sending the summary event of the current process analytic session?
+	if (AnalyticsProvider.GetSessionID().Contains(Session.SessionId)) // The string (GUID) returned by GetSessionID() is surrounded with braces like "{3FEA3232-...}" while Session.SessionId is not -> "3FEA3232-..."
+	{
+		AnalyticsProvider.RecordEvent(TEXT("SessionSummary"), AnalyticsAttributes);
+	}
+	else // The summary was created by another process/instance in a different session. (Ex: Editor sending a summary a prevoulsy crashed instance or CrashReportClientEditor sending it on behalf of the Editor)
+	{
+		// The provider sending a 'summary event' created by another instance/process must parametrize its post request 'as if' it was sent from the instance/session that created it (backend expectation).
+		// Create a new provider to avoid interfering with the current session events. (ex. if another thread sends telemetry at the same time, don't accidently tag it with the wrong SessionID, AppID, etc.).
+		TSharedPtr<IAnalyticsProviderET> TempSummaryProvider = FAnalyticsET::Get().CreateAnalyticsProvider(AnalyticsProvider.GetConfig());
+		
+		// Reconfigure the analytics provider to sent the summary event 'as if' it was sent by the process that created it. This is required by the analytics backend.
+		FGuid SessionGuid;
+		FGuid::Parse(Session.SessionId, SessionGuid);
+		TempSummaryProvider->SetSessionID(SessionGuid.ToString(EGuidFormats::DigitsWithHyphensInBraces)); // Ensure to put back the {} around the GUID.
+		TempSummaryProvider->SetAppID(CopyTemp(Session.AppId));
+		TempSummaryProvider->SetAppVersion(CopyTemp(Session.AppVersion));
+		TempSummaryProvider->SetUserID(CopyTemp(Session.UserId));
+
+		// Send the summary.
+		TempSummaryProvider->RecordEvent(TEXT("SessionSummary"), AnalyticsAttributes);
+
+		// The temporary provider is about to be deleted (going out of scope), ensure it sents its report.
+		TempSummaryProvider->BlockUntilFlushed(2.0f);
+	}
 
 	UE_LOG(LogEditorSessionSummary, Log, TEXT("EditorSessionSummary sent report. Type=%s, SessionId=%s"), *ShutdownTypeString, *SessionIdString);
 }
diff --git a/Engine/Source/Developer/EditorAnalyticsSession/Public/EditorAnalyticsSession.h b/Engine/Source/Developer/EditorAnalyticsSession/Public/EditorAnalyticsSession.h
--- a/Engine/Source/Developer/EditorAnalyticsSession/Public/EditorAnalyticsSession.h
+++ b/Engine/Source/Developer/EditorAnalyticsSession/Public/EditorAnalyticsSession.h
@@ -11,6 +11,10 @@ struct EDITORANALYTICSSESSION_API FEditorAnalyticsSession
 {
 	FString SessionId;
 
+	FString AppId;
+	FString AppVersion;
+	FString UserId;
+
 	FString ProjectName;
 	FString ProjectID;
 	FString ProjectDescription;
diff --git a/Engine/Source/Developer/EditorAnalyticsSession/Public/EditorSessionSummarySender.h b/Engine/Source/Developer/EditorAnalyticsSession/Public/EditorSessionSummarySender.h
--- a/Engine/Source/Developer/EditorAnalyticsSession/Public/EditorSessionSummarySender.h
+++ b/Engine/Source/Developer/EditorAnalyticsSession/Public/EditorSessionSummarySender.h
@@ -7,7 +7,7 @@
 #include "Misc/Optional.h"
 
 struct FEditorAnalyticsSession;
-class IAnalyticsProvider;
+class IAnalyticsProviderET;
 
 /**
   * Sender of SessionSummary events from all editor sessions in-between runs.
@@ -16,7 +16,7 @@ class IAnalyticsProvider;
 class EDITORANALYTICSSESSION_API FEditorSessionSummarySender
 {
 public:
-	FEditorSessionSummarySender(IAnalyticsProvider& InAnalyticsProvider, const FString& InSenderName, const int32 InCurrentSessionProcessId);
+	FEditorSessionSummarySender(IAnalyticsProviderET& InAnalyticsProvider, const FString& InSenderName, const int32 InCurrentSessionProcessId);
 	~FEditorSessionSummarySender();
 
 	void Tick(float DeltaTime);
@@ -32,7 +32,7 @@ private:
 
 private:
 	float HeartbeatTimeElapsed;
-	IAnalyticsProvider& AnalyticsProvider;
+	IAnalyticsProviderET& AnalyticsProvider;
 	FString Sender;
 
 	int32 CurrentSessionProcessId;
diff --git a/Engine/Source/Runtime/Analytics/AnalyticsET/Private/IAnalyticsProviderET.cpp b/Engine/Source/Runtime/Analytics/AnalyticsET/Private/IAnalyticsProviderET.cpp
--- a/Engine/Source/Runtime/Analytics/AnalyticsET/Private/IAnalyticsProviderET.cpp
+++ b/Engine/Source/Runtime/Analytics/AnalyticsET/Private/IAnalyticsProviderET.cpp
@@ -207,7 +207,7 @@ public:
 	virtual void BlockUntilFlushed(float InTimeoutSec) override;
 	virtual ~FAnalyticsProviderET();
 
-	FString GetAPIKey() const { return APIKey; }
+	virtual const FAnalyticsET::Config& GetConfig() const override { return Config; }
 
 private:
 	/**
@@ -226,16 +226,12 @@ private:
 	TSharedRef<IHttpRequest> CreateRequest();
 
 	bool bSessionInProgress;
-	/** ET Game API Key - Get from your account manager */
-	FString APIKey;
-	/** ET API Server */
-	FString APIServer;
+	/** The current configuration (might be updated with respect to the one provided at construction). */
+	FAnalyticsET::Config Config;
 	/** the unique UserID as passed to ET. */
 	FString UserID;
 	/** The session ID */
 	FString SessionID;
-	/** The AppVersion passed to ET. */
-	FString AppVersion;
 	/** Max number of analytics events to cache before pushing to server */
 	const int32 MaxCachedNumEvents;
 	/** Max time that can elapse before pushing cached events to server */
@@ -250,14 +246,7 @@ private:
 	float FlushEventsCountdown;
 	/** Track destructing for unbinding callbacks when firing events at shutdown */
 	bool bInDestructor;
-	/** True to use the legacy backend server protocol that uses URL params. */
-	bool UseLegacyProtocol;
-	/** False to disable dropping events when a flush fails due to network reasons */
-	bool bDropEventsOnFlushFailure;
-	/** AppEnvironment to use. */
-	FString AppEnvironment;
-	/** UploadType to use. */
-	FString UploadType;
+
 	/**
 	* Analytics event entry to be cached
 	*/
@@ -320,9 +309,8 @@ TSharedPtr<IAnalyticsProviderET> FAnalyticsET::CreateAnalyticsProvider(const Con
  * Perform any initialization.
  */
 FAnalyticsProviderET::FAnalyticsProviderET(const FAnalyticsET::Config& ConfigValues)
-	:bSessionInProgress(false)
-	, APIKey(ConfigValues.APIKeyET)
-	, APIServer(ConfigValues.APIServerET)
+	: bSessionInProgress(false)
+	, Config(ConfigValues)
 	, MaxCachedNumEvents(20)
 	, MaxCachedElapsedTime(60.0f)
 	, RetryDelaySecs(120.0f)
@@ -330,12 +318,10 @@ FAnalyticsProviderET::FAnalyticsProviderET(const FAnalyticsET::Config& ConfigVal
 	, bShouldCacheEvents(true)
 	, FlushEventsCountdown(MaxCachedElapsedTime)
 	, bInDestructor(false)
-	, UseLegacyProtocol(ConfigValues.UseLegacyProtocol)
-	, bDropEventsOnFlushFailure(ConfigValues.bDropEventsOnFlushFailure)
 {
-	if (APIKey.IsEmpty() || APIServer.IsEmpty())
+	if (Config.APIKeyET.IsEmpty() || Config.APIServerET.IsEmpty())
 	{
-		UE_LOG(LogAnalytics, Fatal, TEXT("AnalyticsET: APIKey (%s) and APIServer (%s) cannot be empty!"), *APIKey, *APIServer);
+		UE_LOG(LogAnalytics, Fatal, TEXT("AnalyticsET: APIKey (%s) and APIServer (%s) cannot be empty!"), *Config.APIKeyET, *Config.APIServerET);
 	}
 
 	// Set the number of retries to the number of retry URLs that have been passed in.
@@ -351,7 +337,7 @@ FAnalyticsProviderET::FAnalyticsProviderET(const FAnalyticsET::Config& ConfigVal
 	{
 		TArray<FString> TmpAltAPIServers = ConfigValues.AltAPIServersET;
 
-		FString DefaultUrlDomain = FPlatformHttp::GetUrlDomain(APIServer);
+		FString DefaultUrlDomain = FPlatformHttp::GetUrlDomain(Config.APIServerET);
 		if (!TmpAltAPIServers.Contains(DefaultUrlDomain))
 		{
 			TmpAltAPIServers.Insert(DefaultUrlDomain, 0);
@@ -373,25 +359,25 @@ FAnalyticsProviderET::FAnalyticsProviderET(const FAnalyticsET::Config& ConfigVal
 	// make sure that we always start with one control event in the CachedEvents array.
 	CachedEvents.Emplace(FString(), TArray<FAnalyticsEventAttribute>(), false, true);
 
-	UE_LOG(LogAnalytics, Verbose, TEXT("[%s] Initializing ET Analytics provider"), *APIKey);
+	UE_LOG(LogAnalytics, Verbose, TEXT("[%s] Initializing ET Analytics provider"), *Config.APIKeyET);
 
 	// default to FEngineVersion::Current() if one is not provided, append FEngineVersion::Current() otherwise.
 	FString ConfigAppVersion = ConfigValues.AppVersionET;
 	// Allow the cmdline to force a specific AppVersion so it can be set dynamically.
 	FParse::Value(FCommandLine::Get(), TEXT("ANALYTICSAPPVERSION="), ConfigAppVersion, false);
-	AppVersion = ConfigAppVersion.IsEmpty()
+	Config.AppVersionET = ConfigAppVersion.IsEmpty()
 		? FString(FApp::GetBuildVersion())
 		: ConfigAppVersion.Replace(TEXT("%VERSION%"), FApp::GetBuildVersion(), ESearchCase::CaseSensitive);
 
-	UE_LOG(LogAnalytics, Log, TEXT("[%s] APIServer = %s. AppVersion = %s"), *APIKey, *APIServer, *AppVersion);
+	UE_LOG(LogAnalytics, Log, TEXT("[%s] APIServer = %s. AppVersion = %s"), *Config.APIKeyET, *Config.APIServerET, *Config.AppVersionET);
 
 	// only need these if we are using the data router protocol.
-	if (!UseLegacyProtocol)
+	if (!Config.UseLegacyProtocol)
 	{
-		AppEnvironment = ConfigValues.AppEnvironment.IsEmpty()
+		Config.AppEnvironment = ConfigValues.AppEnvironment.IsEmpty()
 			? FAnalyticsET::Config::GetDefaultAppEnvironment()
 			: ConfigValues.AppEnvironment;
-		UploadType = ConfigValues.UploadType.IsEmpty()
+		Config.UploadType = ConfigValues.UploadType.IsEmpty()
 			? FAnalyticsET::Config::GetDefaultUploadType()
 			: ConfigValues.UploadType;
 	}
@@ -449,7 +435,7 @@ bool FAnalyticsProviderET::Tick(float DeltaSeconds)
 
 FAnalyticsProviderET::~FAnalyticsProviderET()
 {
-	UE_LOG(LogAnalytics, Verbose, TEXT("[%s] Destroying ET Analytics provider"), *APIKey);
+	UE_LOG(LogAnalytics, Verbose, TEXT("[%s] Destroying ET Analytics provider"), *Config.APIKeyET);
 	bInDestructor = true;
 	EndSession();
 }
@@ -466,7 +452,7 @@ bool FAnalyticsProviderET::StartSession(const TArray<FAnalyticsEventAttribute>&
 
 bool FAnalyticsProviderET::StartSession(TArray<FAnalyticsEventAttribute>&& Attributes)
 {
-	UE_LOG(LogAnalytics, Log, TEXT("[%s] AnalyticsET::StartSession"), *APIKey);
+	UE_LOG(LogAnalytics, Log, TEXT("[%s] AnalyticsET::StartSession"), *Config.APIKeyET);
 
 	// end/flush previous session before staring new one
 	if (bSessionInProgress)
@@ -544,10 +530,10 @@ void FAnalyticsProviderET::FlushEvents()
 		// We can do this without actually copying the array this way.
 		TArray<FAnalyticsEventAttribute> CurrentDefaultAttributes;
 
-		if (!UseLegacyProtocol)
+		if (!Config.UseLegacyProtocol)
 		{
 			TSharedPtr< TArray<FAnalyticsEventEntry> > FlushedEvents;
-			if (!bDropEventsOnFlushFailure)
+			if (!Config.bDropEventsOnFlushFailure)
 			{
 				FlushedEvents = TSharedPtr< TArray<FAnalyticsEventEntry> >(new TArray<FAnalyticsEventEntry>());
 				FlushedEvents->Reserve(CachedEvents.Num());
@@ -647,11 +633,11 @@ void FAnalyticsProviderET::FlushEvents()
 
 			// UrlEncode NOTE: need to concatenate everything
 			FString URLPath  = TEXT("datarouter/api/v1/public/data?SessionID=") + FPlatformHttp::UrlEncode(SessionID);
-					URLPath += TEXT("&AppID=") + FPlatformHttp::UrlEncode(APIKey);
-					URLPath += TEXT("&AppVersion=") + FPlatformHttp::UrlEncode(AppVersion);
+					URLPath += TEXT("&AppID=") + FPlatformHttp::UrlEncode(Config.APIKeyET);
+					URLPath += TEXT("&AppVersion=") + FPlatformHttp::UrlEncode(Config.AppVersionET);
 					URLPath += TEXT("&UserID=") + FPlatformHttp::UrlEncode(UserID);
-					URLPath += TEXT("&AppEnvironment=") + FPlatformHttp::UrlEncode(AppEnvironment);
-					URLPath += TEXT("&UploadType=") + FPlatformHttp::UrlEncode(UploadType);
+					URLPath += TEXT("&AppEnvironment=") + FPlatformHttp::UrlEncode(Config.AppEnvironment);
+					URLPath += TEXT("&UploadType=") + FPlatformHttp::UrlEncode(Config.UploadType);
 			PayloadSize = URLPath.Len() + Payload.Len();
 
 			if (UE_LOG_ACTIVE(LogAnalytics, VeryVerbose))
@@ -659,13 +645,13 @@ void FAnalyticsProviderET::FlushEvents()
 				// Recreate the URLPath for logging because we do not want to escape the parameters when logging.
 				// We cannot simply UrlEncode the entire Path after logging it because UrlEncode(Params) != UrlEncode(Param1) & UrlEncode(Param2) ...
 				FString LogString = FString::Printf(TEXT("[%s] AnalyticsET URL:datarouter/api/v1/public/data?SessionID=%s&AppID=%s&AppVersion=%s&UserID=%s&AppEnvironment=%s&UploadType=%s. Payload:%s"),
-					*APIKey,
+					*Config.APIKeyET,
 					*SessionID,
-					*APIKey,
-					*AppVersion,
+					*Config.APIKeyET,
+					*Config.AppVersionET,
 					*UserID,
-					*AppEnvironment,
-					*UploadType,
+					*Config.AppEnvironment,
+					*Config.UploadType,
 					*Payload);
 				UE_LOG(LogAnalytics, VeryVerbose, TEXT("%s"), *LogString);
 			}
@@ -675,7 +661,7 @@ void FAnalyticsProviderET::FlushEvents()
 				// Create/send Http request for an event
 				TSharedRef<IHttpRequest> HttpRequest = CreateRequest();
 				HttpRequest->SetHeader(TEXT("Content-Type"), TEXT("application/json; charset=utf-8"));
-				HttpRequest->SetURL(APIServer / URLPath);
+				HttpRequest->SetURL(Config.APIServerET / URLPath);
 				HttpRequest->SetVerb(TEXT("POST"));
 				HttpRequest->SetContentAsString(Payload);
 
@@ -724,10 +710,10 @@ void FAnalyticsProviderET::FlushEvents()
 
 					// log out the un-encoded values to make reading the log easier.
 					UE_LOG(LogAnalytics, VeryVerbose, TEXT("[%s] AnalyticsET URL:SendEvent.1?SessionID=%s&AppID=%s&AppVersion=%s&UserID=%s&EventName=%s%s"),
-						*APIKey,
+						*Config.APIKeyET,
 						*SessionID,
-						*APIKey,
-						*AppVersion,
+						*Config.APIKeyET,
+						*Config.AppVersionET,
 						*UserID,
 						*Event.EventName,
 						*EventParams);
@@ -737,10 +723,10 @@ void FAnalyticsProviderET::FlushEvents()
 					HttpRequest->SetHeader(TEXT("Content-Type"), TEXT("text/plain"));
 
 					// Don't need to URL encode the APIServer or the EventParams, which are already encoded, and contain parameter separaters that we DON'T want encoded.
-					FString URLPath  = APIServer;
+					FString URLPath  = Config.APIServerET;
 							URLPath += TEXT("SendEvent.1?SessionID=") + FPlatformHttp::UrlEncode(SessionID);
-							URLPath += TEXT("&AppID=") + FPlatformHttp::UrlEncode(APIKey);
-							URLPath += TEXT("&AppVersion=") + FPlatformHttp::UrlEncode(AppVersion);
+							URLPath += TEXT("&AppID=") + FPlatformHttp::UrlEncode(Config.APIKeyET);
+							URLPath += TEXT("&AppVersion=") + FPlatformHttp::UrlEncode(Config.AppVersionET);
 							URLPath += TEXT("&UserID=") + FPlatformHttp::UrlEncode(UserID);
 							URLPath += TEXT("&EventName=") + FPlatformHttp::UrlEncode(Event.EventName);
 							URLPath += EventParams;
@@ -768,11 +754,11 @@ void FAnalyticsProviderET::FlushEvents()
 
 void FAnalyticsProviderET::SetAppID(FString&& InAppID)
 {
-	if (APIKey != InAppID)
+	if (Config.APIKeyET != InAppID)
 	{
 		// Flush any cached events that would be using the old AppID.
 		FlushEvents();
-		APIKey = MoveTemp(InAppID);
+		Config.APIKeyET = MoveTemp(InAppID);
 	}
 }
 
@@ -783,23 +769,23 @@ void FAnalyticsProviderET::SetAppVersion(FString&& InAppVersion)
 		? FString(FApp::GetBuildVersion())
 		: InAppVersion.Replace(TEXT("%VERSION%"), FApp::GetBuildVersion(), ESearchCase::CaseSensitive);
 
-	if (AppVersion != InAppVersion)
+	if (Config.AppVersionET != InAppVersion)
 	{
-		UE_LOG(LogAnalytics, Log, TEXT("[%s] Updating AppVersion to %s from old value of %s"), *APIKey, *InAppVersion, *AppVersion);
+		UE_LOG(LogAnalytics, Log, TEXT("[%s] Updating AppVersion to %s from old value of %s"), *Config.APIKeyET, *InAppVersion, *Config.AppVersionET);
 		// Flush any cached events that would be using the old AppVersion.
 		FlushEvents();
-		AppVersion = MoveTemp(InAppVersion);
+		Config.AppVersionET = MoveTemp(InAppVersion);
 	}
 }
 
 const FString& FAnalyticsProviderET::GetAppID() const
 {
-	return APIKey;
+	return Config.APIKeyET;
 }
 
 const FString& FAnalyticsProviderET::GetAppVersion() const
 {
-	return AppVersion;
+	return Config.AppVersionET;
 }
 
 void FAnalyticsProviderET::SetUserID(const FString& InUserID)
@@ -807,14 +793,14 @@ void FAnalyticsProviderET::SetUserID(const FString& InUserID)
 	// command-line specified user ID overrides all attempts to reset it.
 	if (!FParse::Value(FCommandLine::Get(), TEXT("ANALYTICSUSERID="), UserID, false))
 	{
-		UE_LOG(LogAnalytics, Log, TEXT("[%s] SetUserId %s"), *APIKey, *InUserID);
+		UE_LOG(LogAnalytics, Log, TEXT("[%s] SetUserId %s"), *Config.APIKeyET, *InUserID);
 		// Flush any cached events that would be using the old UserID.
 		FlushEvents();
 		UserID = InUserID;
 	}
 	else if (UserID != InUserID)
 	{
-		UE_LOG(LogAnalytics, Log, TEXT("[%s] Overriding SetUserId %s with cmdline UserId of %s."), *APIKey, *InUserID, *UserID);
+		UE_LOG(LogAnalytics, Log, TEXT("[%s] Overriding SetUserId %s with cmdline UserId of %s."), *Config.APIKeyET, *InUserID, *UserID);
 	}
 }
 
@@ -835,7 +821,7 @@ bool FAnalyticsProviderET::SetSessionID(const FString& InSessionID)
 		// Flush any cached events that would be using the old SessionID.
 		FlushEvents();
 		SessionID = InSessionID;
-		UE_LOG(LogAnalytics, Log, TEXT("[%s] Forcing SessionID to %s."), *APIKey, *SessionID);
+		UE_LOG(LogAnalytics, Log, TEXT("[%s] Forcing SessionID to %s."), *Config.APIKeyET, *SessionID);
 	}
 	return true;
 }
@@ -867,7 +853,7 @@ void FAnalyticsProviderET::RecordEvent(FString EventName, TArray<FAnalyticsEvent
 
 void FAnalyticsProviderET::RecordEventJson(FString EventName, TArray<FAnalyticsEventAttribute>&& AttributesJson)
 {
-	checkf(!UseLegacyProtocol, TEXT("Cannot use Json events with legacy protocol"));
+	checkf(!Config.UseLegacyProtocol, TEXT("Cannot use Json events with legacy protocol"));
 
 	// fire any callbacks
 	for (const auto& Cb : EventRecordedCallbacks)
@@ -921,7 +907,7 @@ void FAnalyticsProviderET::EventRequestComplete(FHttpRequestPtr HttpRequest, FHt
 	bool bEventsDelivered = false;
 	if (HttpResponse.IsValid())
 	{
-		UE_LOG(LogAnalytics, VeryVerbose, TEXT("[%s] ET response for [%s]. Code: %d. Payload: %s"), *APIKey, *HttpRequest->GetURL(), HttpResponse->GetResponseCode(), *HttpResponse->GetContentAsString());
+		UE_LOG(LogAnalytics, VeryVerbose, TEXT("[%s] ET response for [%s]. Code: %d. Payload: %s"), *Config.APIKeyET, *HttpRequest->GetURL(), HttpResponse->GetResponseCode(), *HttpResponse->GetContentAsString());
 		if (EHttpResponseCodes::IsOk(HttpResponse->GetResponseCode()))
 		{
 			bEventsDelivered = true;
@@ -929,7 +915,7 @@ void FAnalyticsProviderET::EventRequestComplete(FHttpRequestPtr HttpRequest, FHt
 	}
 	else
 	{
-		UE_LOG(LogAnalytics, VeryVerbose, TEXT("[%s] ET response for [%s]. No response"), *APIKey, *HttpRequest->GetURL());
+		UE_LOG(LogAnalytics, VeryVerbose, TEXT("[%s] ET response for [%s]. No response"), *Config.APIKeyET, *HttpRequest->GetURL());
 	}
 
 	// if the events were not delivered
@@ -954,14 +940,14 @@ void FAnalyticsProviderET::EventRequestComplete(FHttpRequestPtr HttpRequest, FHt
 			// if we're being super spammy or have been offline forever, just leave it at the ET.DroppedSubmission event
 			if (bShouldCacheEvents && CachedEvents.Num() < 256)
 			{
-				UE_LOG(LogAnalytics, Log, TEXT("[%s] ET Requeuing %d analytics events due to failure to send"), *APIKey, FlushedEvents->Num());
+				UE_LOG(LogAnalytics, Log, TEXT("[%s] ET Requeuing %d analytics events due to failure to send"), *Config.APIKeyET, FlushedEvents->Num());
 
 				// put them at the beginning since it should include a default attributes entry and we don't want to change the current default attributes
 				CachedEvents.Insert(*FlushedEvents, 0);
 			}
 			else
 			{
-				UE_LOG(LogAnalytics, Error, TEXT("[%s] ET dropping %d analytics events due to too many in queue (%d)"), *APIKey, FlushedEvents->Num(), CachedEvents.Num());
+				UE_LOG(LogAnalytics, Error, TEXT("[%s] ET dropping %d analytics events due to too many in queue (%d)"), *Config.APIKeyET, FlushedEvents->Num(), CachedEvents.Num());
 			}
 		}
 	}
@@ -970,7 +956,7 @@ void FAnalyticsProviderET::EventRequestComplete(FHttpRequestPtr HttpRequest, FHt
 void FAnalyticsProviderET::SetURLEndpoint(const FString& UrlEndpoint, const TArray<FString>& AltDomains)
 {
 	FlushEvents();
-	APIServer = UrlEndpoint;
+	Config.APIServerET = UrlEndpoint;
 
 	// Set the number of retries to the number of retry URLs that have been passed in.
 	uint32 RetryLimitCount = AltDomains.Num();
@@ -982,7 +968,7 @@ void FAnalyticsProviderET::SetURLEndpoint(const FString& UrlEndpoint, const TArr
 	// If we have retry domains defined, insert the default domain into the list
 	if (RetryLimitCount > 0)
 	{
-		FString DefaultUrlDomain = FPlatformHttp::GetUrlDomain(APIServer);
+		FString DefaultUrlDomain = FPlatformHttp::GetUrlDomain(Config.APIServerET);
 		if (!TmpAltAPIServers.Contains(DefaultUrlDomain))
 		{
 			TmpAltAPIServers.Insert(DefaultUrlDomain, 0);
diff --git a/Engine/Source/Runtime/Analytics/AnalyticsET/Public/IAnalyticsProviderET.h b/Engine/Source/Runtime/Analytics/AnalyticsET/Public/IAnalyticsProviderET.h
--- a/Engine/Source/Runtime/Analytics/AnalyticsET/Public/IAnalyticsProviderET.h
+++ b/Engine/Source/Runtime/Analytics/AnalyticsET/Public/IAnalyticsProviderET.h
@@ -4,6 +4,7 @@
 
 #include "CoreMinimal.h"
 #include "AnalyticsEventAttribute.h"
+#include "AnalyticsET.h"
 #include "Interfaces/IAnalyticsProvider.h"
 
 /** ET specific analytics provider instance. Exposes additional APIs to support Json-based events. */
@@ -119,4 +120,9 @@ public:
 	* Blocks execution in the thread until all events have been flushed to the network.
 	*/
 	virtual void BlockUntilFlushed(float InTimeoutSec) = 0;
+
+	/**
+	 * Return the current provider configuration.
+	 */
+	virtual const FAnalyticsET::Config& GetConfig() const = 0;
 };
diff --git a/Engine/Source/Runtime/Engine/Private/Analytics/EditorSessionSummaryWriter.cpp b/Engine/Source/Runtime/Engine/Private/Analytics/EditorSessionSummaryWriter.cpp
--- a/Engine/Source/Runtime/Engine/Private/Analytics/EditorSessionSummaryWriter.cpp
+++ b/Engine/Source/Runtime/Engine/Private/Analytics/EditorSessionSummaryWriter.cpp
@@ -13,7 +13,7 @@
 #include "Engine/Engine.h"
 #include "Framework/Application/SlateApplication.h"
 #include "HAL/PlatformProcess.h"
-#include "Interfaces/IAnalyticsProvider.h"
+#include "IAnalyticsProviderET.h"
 #include "Interfaces/IPluginManager.h"
 #include "Interfaces/IProjectManager.h"
 #include "Logging/LogMacros.h"
@@ -201,21 +201,29 @@ void FEditorSessionSummaryWriter::Shutdown()
 
 FEditorAnalyticsSession* FEditorSessionSummaryWriter::CreateCurrentSession() const
 {
+	check(FEngineAnalytics::IsAvailable()); // The function assumes the caller checked it before calling.
+
 	FEditorAnalyticsSession* Session = new FEditorAnalyticsSession();
+	IAnalyticsProviderET& AnalyticProvider = FEngineAnalytics::GetProvider();
 
 	FGuid SessionId;
-	if (FGuid::Parse(FEngineAnalytics::GetProvider().GetSessionID(), SessionId))
+	if (FGuid::Parse(AnalyticProvider.GetSessionID(), SessionId))
 	{
 		// convert session GUID to one without braces or other chars that might not be suitable for storage
 		Session->SessionId = SessionId.ToString(EGuidFormats::DigitsWithHyphens);
 	}
 	else
 	{
-		Session->SessionId = FEngineAnalytics::GetProvider().GetSessionID();
+		Session->SessionId = AnalyticProvider.GetSessionID();
 	}
 
 	const UGeneralProjectSettings& ProjectSettings = *GetDefault<UGeneralProjectSettings>();
 
+	// Remember the AppId/AppVersion/UserId used during this session. They will be used if the summary is sent from another process/instance.
+	Session->AppId = AnalyticProvider.GetAppID();
+	Session->AppVersion = AnalyticProvider.GetAppVersion();
+	Session->UserId = AnalyticProvider.GetUserID();
+
 	Session->PlatformProcessID = FPlatformProcess::GetCurrentProcessId();
 	Session->ProjectName = ProjectSettings.ProjectName;
 	Session->ProjectID = ProjectSettings.ProjectID.ToString(EGuidFormats::DigitsWithHyphens);
diff --git a/Engine/Source/Runtime/Engine/Private/EngineAnalytics.cpp b/Engine/Source/Runtime/Engine/Private/EngineAnalytics.cpp
--- a/Engine/Source/Runtime/Engine/Private/EngineAnalytics.cpp
+++ b/Engine/Source/Runtime/Engine/Private/EngineAnalytics.cpp
@@ -89,13 +89,13 @@ static TSharedPtr<IAnalyticsProviderET> CreateEpicAnalyticsProvider()
 	return FAnalyticsET::Get().CreateAnalyticsProvider(Config);
 }
 
-IAnalyticsProvider& FEngineAnalytics::GetProvider()
+IAnalyticsProviderET& FEngineAnalytics::GetProvider()
 {
 	checkf(bIsInitialized && IsAvailable(), TEXT("FEngineAnalytics::GetProvider called outside of Initialize/Shutdown."));
 
 	return *Analytics.Get();
 }
- 
+
 void FEngineAnalytics::Initialize()
 {
 	checkf(!bIsInitialized, TEXT("FEngineAnalytics::Initialize called more than once."));
diff --git a/Engine/Source/Runtime/Engine/Public/EngineAnalytics.h b/Engine/Source/Runtime/Engine/Public/EngineAnalytics.h
--- a/Engine/Source/Runtime/Engine/Public/EngineAnalytics.h
+++ b/Engine/Source/Runtime/Engine/Public/EngineAnalytics.h
@@ -3,6 +3,7 @@
 #pragma once
 
 #include "CoreMinimal.h"
+#include "IAnalyticsProviderET.h" // NOTE: Consider changing the code to replace IAnalyticProvider.h by IAnalyticProviderET.h
 
 class FEngineSessionManager;
 class IAnalyticsProvider;
@@ -26,7 +27,7 @@ public:
 	 * Return the provider instance. Not valid outside of Initialize/Shutdown calls.
 	 * Note: must check IsAvailable() first else this code will assert if the provider is not valid.
 	 */
-	static ENGINE_API IAnalyticsProvider& GetProvider();
+	static ENGINE_API IAnalyticsProviderET& GetProvider();
 
 	/** Helper function to determine if the provider is valid. */
 	static ENGINE_API bool IsAvailable() { return Analytics.IsValid(); }
