prevent crash in TJsonWriter by checking if WriteObjectStart can be executed

diff --git a/Engine/Source/Runtime/Analytics/AnalyticsET/Private/IAnalyticsProviderET.cpp b/Engine/Source/Runtime/Analytics/AnalyticsET/Private/IAnalyticsProviderET.cpp
--- a/Engine/Source/Runtime/Analytics/AnalyticsET/Private/IAnalyticsProviderET.cpp
+++ b/Engine/Source/Runtime/Analytics/AnalyticsET/Private/IAnalyticsProviderET.cpp
@@ -550,55 +550,36 @@ void FAnalyticsProviderET::FlushEvents()
 			{
 				QUICK_SCOPE_CYCLE_COUNTER(STAT_FlushEventsJson);
 				TSharedRef< TJsonWriter<TCHAR, TCondensedJsonPrintPolicy<TCHAR> > > JsonWriter = TJsonWriterFactory<TCHAR, TCondensedJsonPrintPolicy<TCHAR> >::Create(&Payload);
-				JsonWriter->WriteObjectStart();
-				JsonWriter->WriteArrayStart(TEXT("Events"));
-				for (FAnalyticsEventEntry& Entry : CachedEvents)
+				if (JsonWriter->CanWriteObjectStart())
 				{
-					if (Entry.bIsDefaultAttributes)
+					JsonWriter->WriteObjectStart();
+					JsonWriter->WriteArrayStart(TEXT("Events"));
+					for (FAnalyticsEventEntry& Entry : CachedEvents)
 					{
-						// This is the default attributes, so update the array.
-						if (FlushedEvents.IsValid())
+						if (Entry.bIsDefaultAttributes)
 						{
-							CurrentDefaultAttributes = Entry.Attributes; // need to copy
-							FlushedEvents->Emplace(MoveTemp(Entry));
-						}
-						else
-						{
-							CurrentDefaultAttributes = MoveTemp(Entry.Attributes);
-						}
-					}
-					else
-					{
-						++EventCount;
-
-						// event entry
-						JsonWriter->WriteObjectStart();
-						JsonWriter->WriteValue(TEXT("EventName"), Entry.EventName);
-						FString DateOffset = (CurrentTime - Entry.TimeStamp).ToString();
-						JsonWriter->WriteValue(TEXT("DateOffset"), DateOffset);
-						// default attributes for this event
-						for (const FAnalyticsEventAttribute& Attr : CurrentDefaultAttributes)
-						{
-							switch (Attr.AttrType)
+							// This is the default attributes, so update the array.
+							if (FlushedEvents.IsValid())
 							{
-							case FAnalyticsEventAttribute::AttrTypeEnum::String:
-								JsonWriter->WriteValue(Attr.AttrName, Attr.AttrValueString);
-								break;
-							case FAnalyticsEventAttribute::AttrTypeEnum::Number:
-								JsonWriter->WriteValue(Attr.AttrName, Attr.ToString());
-								break;
-							case FAnalyticsEventAttribute::AttrTypeEnum::Boolean:
-								JsonWriter->WriteValue(Attr.AttrName, Attr.ToString());
-								break;
-							case FAnalyticsEventAttribute::AttrTypeEnum::JsonFragment:
-								JsonWriter->WriteRawJSONValue(Attr.AttrName, Attr.AttrValueString);
-								break;
+								CurrentDefaultAttributes = Entry.Attributes; // need to copy
+								FlushedEvents->Emplace(MoveTemp(Entry));
+							}
+							else
+							{
+								CurrentDefaultAttributes = MoveTemp(Entry.Attributes);
 							}
 						}
-						// optional attributes for this event
-						if (!Entry.bIsJsonEvent)
+						else
 						{
-							for (const FAnalyticsEventAttribute& Attr : Entry.Attributes)
+							++EventCount;
+
+							// event entry
+							JsonWriter->WriteObjectStart();
+							JsonWriter->WriteValue(TEXT("EventName"), Entry.EventName);
+							FString DateOffset = (CurrentTime - Entry.TimeStamp).ToString();
+							JsonWriter->WriteValue(TEXT("DateOffset"), DateOffset);
+							// default attributes for this event
+							for (const FAnalyticsEventAttribute& Attr : CurrentDefaultAttributes)
 							{
 								switch (Attr.AttrType)
 								{
@@ -616,26 +597,49 @@ void FAnalyticsProviderET::FlushEvents()
 									break;
 								}
 							}
-						}
-						else
-						{
-							for (const FAnalyticsEventAttribute& Attr : Entry.Attributes)
+							// optional attributes for this event
+							if (!Entry.bIsJsonEvent)
 							{
-								JsonWriter->WriteRawJSONValue(Attr.AttrName, Attr.AttrValueString);
+								for (const FAnalyticsEventAttribute& Attr : Entry.Attributes)
+								{
+									switch (Attr.AttrType)
+									{
+									case FAnalyticsEventAttribute::AttrTypeEnum::String:
+										JsonWriter->WriteValue(Attr.AttrName, Attr.AttrValueString);
+										break;
+									case FAnalyticsEventAttribute::AttrTypeEnum::Number:
+										JsonWriter->WriteValue(Attr.AttrName, Attr.ToString());
+										break;
+									case FAnalyticsEventAttribute::AttrTypeEnum::Boolean:
+										JsonWriter->WriteValue(Attr.AttrName, Attr.ToString());
+										break;
+									case FAnalyticsEventAttribute::AttrTypeEnum::JsonFragment:
+										JsonWriter->WriteRawJSONValue(Attr.AttrName, Attr.AttrValueString);
+										break;
+									}
+								}
 							}
-						}
-						JsonWriter->WriteObjectEnd();
+							else
+							{
+								for (const FAnalyticsEventAttribute& Attr : Entry.Attributes)
+								{
+									JsonWriter->WriteRawJSONValue(Attr.AttrName, Attr.AttrValueString);
+								}
+							}
+							JsonWriter->WriteObjectEnd();
 
-						// move the entry into the flushed
-						if (FlushedEvents.IsValid())
-						{
-							FlushedEvents->Emplace(MoveTemp(Entry));
+							// move the entry into the flushed
+							if (FlushedEvents.IsValid())
+							{
+								FlushedEvents->Emplace(MoveTemp(Entry));
+							}
 						}
 					}
-				}
 
-				JsonWriter->WriteArrayEnd();
-				JsonWriter->WriteObjectEnd();
+					JsonWriter->WriteArrayEnd();
+					JsonWriter->WriteObjectEnd();
+				} // CanWriteObjectStart
+
 				JsonWriter->Close();
 			}
 
diff --git a/Engine/Source/Runtime/Json/Public/Serialization/JsonWriter.h b/Engine/Source/Runtime/Json/Public/Serialization/JsonWriter.h
--- a/Engine/Source/Runtime/Json/Public/Serialization/JsonWriter.h
+++ b/Engine/Source/Runtime/Json/Public/Serialization/JsonWriter.h
@@ -29,6 +29,11 @@ public:
 
 	FORCEINLINE int32 GetIndentLevel() const { return IndentLevel; }
 
+	bool CanWriteObjectStart() const
+	{
+		return CanWriteObjectWithoutIdentifier();
+	}
+
 	void WriteObjectStart()
 	{
 		check(CanWriteObjectWithoutIdentifier());
