Added a ability for an IAnalyticsProviderET to check a user-supplied callback function before sending an event.

#tests PIE, client, server

diff --git a/Engine/Source/Runtime/Analytics/AnalyticsET/Private/IAnalyticsProviderET.cpp b/Engine/Source/Runtime/Analytics/AnalyticsET/Private/IAnalyticsProviderET.cpp
--- a/Engine/Source/Runtime/Analytics/AnalyticsET/Private/IAnalyticsProviderET.cpp
+++ b/Engine/Source/Runtime/Analytics/AnalyticsET/Private/IAnalyticsProviderET.cpp
@@ -197,6 +197,7 @@ public:
 	virtual FString GetSessionID() const override;
 	virtual bool SetSessionID(const FString& InSessionID) override;
 
+	virtual bool ShouldRecordEvent(const FString& EventName) const override;
 	virtual void RecordEvent(const FString& EventName, const TArray<FAnalyticsEventAttribute>& Attributes) override;
 	virtual void RecordEvent(FString EventName, TArray<FAnalyticsEventAttribute>&& Attributes) override;
 	virtual void RecordEventJson(FString EventName, TArray<FAnalyticsEventAttribute>&& AttributesJson) override;
@@ -206,6 +207,7 @@ public:
 
 	virtual void SetURLEndpoint(const FString& UrlEndpoint, const TArray<FString>& AltDomains) override;
 	virtual void BlockUntilFlushed(float InTimeoutSec) override;
+	virtual void SetShouldRecordEventFunc(const ShouldRecordEventFunction& InShouldRecordEventFunc) override;
 	virtual ~FAnalyticsProviderET();
 
 	virtual const FAnalyticsET::Config& GetConfig() const override { return Config; }
@@ -286,6 +288,9 @@ private:
 
 	TArray<OnEventRecorded> EventRecordedCallbacks;
 
+	/** Event filter function */
+	ShouldRecordEventFunction ShouldRecordEventFunc;
+
 	/**
 	* Delegate called when an event Http request completes
 	*/
@@ -303,7 +308,7 @@ TSharedPtr<IAnalyticsProviderET> FAnalyticsET::CreateAnalyticsProvider(const Con
 		UE_LOG(LogAnalytics, Warning, TEXT("CreateAnalyticsProvider config not contain required parameter %s"), *Config::GetKeyNameForAPIKey());
 		return NULL;
 	}
-	return TSharedPtr<IAnalyticsProviderET>(new FAnalyticsProviderET(ConfigValues));
+	return MakeShared<FAnalyticsProviderET>(ConfigValues);
 }
 
 /**
@@ -838,6 +843,11 @@ bool FAnalyticsProviderET::SetSessionID(const FString& InSessionID)
 	return true;
 }
 
+bool FAnalyticsProviderET::ShouldRecordEvent(const FString& EventName) const
+{
+	return !ShouldRecordEventFunc || ShouldRecordEventFunc(*this, EventName);
+}
+
 void FAnalyticsProviderET::RecordEvent(const FString& EventName, const TArray<FAnalyticsEventAttribute>& Attributes)
 {
 	// Have to copy Attributes array because this doesn't come in as an rvalue ref.
@@ -846,20 +856,24 @@ void FAnalyticsProviderET::RecordEvent(const FString& EventName, const TArray<FA
 
 void FAnalyticsProviderET::RecordEvent(FString EventName, TArray<FAnalyticsEventAttribute>&& Attributes)
 {
-	// fire any callbacks
-	for (const auto& Cb : EventRecordedCallbacks)
+	// let higher level code filter the decision of whether to send the event
+	if (ShouldRecordEvent(EventName))
 	{
-		Cb(EventName, Attributes, false);
-	}
+		// fire any callbacks
+		for (const auto& Cb : EventRecordedCallbacks)
+		{
+			Cb(EventName, Attributes, false);
+		}
 
-	// There are much better ways to do this, but since most events are recorded and handled on the same (game) thread,
-	// this is probably mostly fine for now, and simply favoring not crashing at the moment
-	FScopeLock ScopedLock(&CachedEventsCS);
-	CachedEvents.Emplace(MoveTemp(EventName), MoveTemp(Attributes), false, false);
-	// if we aren't caching events, flush immediately. This is really only for debugging as it will significantly affect bandwidth.
-	if (!bShouldCacheEvents)
-	{
-		FlushEvents();
+		// There are much better ways to do this, but since most events are recorded and handled on the same (game) thread,
+		// this is probably mostly fine for now, and simply favoring not crashing at the moment
+		FScopeLock ScopedLock(&CachedEventsCS);
+		CachedEvents.Emplace(MoveTemp(EventName), MoveTemp(Attributes), false, false);
+		// if we aren't caching events, flush immediately. This is really only for debugging as it will significantly affect bandwidth.
+		if (!bShouldCacheEvents)
+		{
+			FlushEvents();
+		}
 	}
 }
 
@@ -867,20 +881,24 @@ void FAnalyticsProviderET::RecordEventJson(FString EventName, TArray<FAnalyticsE
 {
 	checkf(!Config.UseLegacyProtocol, TEXT("Cannot use Json events with legacy protocol"));
 
-	// fire any callbacks
-	for (const auto& Cb : EventRecordedCallbacks)
+	// let higher level code filter the decision of whether to send the event
+	if (ShouldRecordEvent(EventName))
 	{
-		Cb(EventName, AttributesJson, true);
-	}
+		// fire any callbacks
+		for (const auto& Cb : EventRecordedCallbacks)
+		{
+			Cb(EventName, AttributesJson, true);
+		}
 
-	// There are much better ways to do this, but since most events are recorded and handled on the same (game) thread,
-	// this is probably mostly fine for now, and simply favoring not crashing at the moment
-	FScopeLock ScopedLock(&CachedEventsCS);
-	CachedEvents.Emplace(MoveTemp(EventName), MoveTemp(AttributesJson), true, false);
-	// if we aren't caching events, flush immediately. This is really only for debugging as it will significantly affect bandwidth.
-	if (!bShouldCacheEvents)
-	{
-		FlushEvents();
+		// There are much better ways to do this, but since most events are recorded and handled on the same (game) thread,
+		// this is probably mostly fine for now, and simply favoring not crashing at the moment
+		FScopeLock ScopedLock(&CachedEventsCS);
+		CachedEvents.Emplace(MoveTemp(EventName), MoveTemp(AttributesJson), true, false);
+		// if we aren't caching events, flush immediately. This is really only for debugging as it will significantly affect bandwidth.
+		if (!bShouldCacheEvents)
+		{
+			FlushEvents();
+		}
 	}
 }
 
@@ -999,3 +1017,8 @@ void FAnalyticsProviderET::BlockUntilFlushed(float InTimeoutSec)
 	FlushEvents();
 	HttpRetryManager->BlockUntilFlushed(InTimeoutSec);
 }
+
+void FAnalyticsProviderET::SetShouldRecordEventFunc(const ShouldRecordEventFunction& InShouldRecordEventFunc)
+{
+	ShouldRecordEventFunc = InShouldRecordEventFunc;
+}
diff --git a/Engine/Source/Runtime/Analytics/AnalyticsET/Public/IAnalyticsProviderET.h b/Engine/Source/Runtime/Analytics/AnalyticsET/Public/IAnalyticsProviderET.h
--- a/Engine/Source/Runtime/Analytics/AnalyticsET/Public/IAnalyticsProviderET.h
+++ b/Engine/Source/Runtime/Analytics/AnalyticsET/Public/IAnalyticsProviderET.h
@@ -7,7 +7,7 @@
 #include "AnalyticsET.h"
 #include "Interfaces/IAnalyticsProvider.h"
 
-/** ET specific analytics provider instance. Exposes additional APIs to support Json-based events. */
+/** ET specific analytics provider instance. Exposes additional APIs to support Json-based events, using move-semantics, and allowing events to be disabled (generally via hotfixing). */
 class IAnalyticsProviderET : public IAnalyticsProvider
 {
 public:
@@ -57,6 +57,14 @@ public:
 	*/
 	virtual bool StartSession(FString InSessionID, TArray<FAnalyticsEventAttribute>&& Attributes) = 0;
 
+	/**
+	* Allows higher level code to abort logic to set up for a RecordEvent call by checking the filter that will be used to send the event first.
+	*
+	* @param EventName			The name of the event.
+	* @return true if the event will be recorded using the currently installed ShouldRecordEvent function
+	*/
+	virtual bool ShouldRecordEvent(const FString& EventName) const = 0;
+
 	/**
 	* Optimization for RecordEvent that avoids the array copy using rvalue references.
 	*
@@ -132,4 +140,10 @@ public:
 	 * Return the current provider configuration.
 	 */
 	virtual const FAnalyticsET::Config& GetConfig() const = 0;
+
+	/** Callback used before any event is actually sent. Allows higher level code to disable events. */
+	typedef TFunction<bool(const IAnalyticsProviderET& ThisProvider, const FString& EventName)> ShouldRecordEventFunction;
+
+	/** Set an event filter to dynamically control whether an event should be sent. */
+	virtual void SetShouldRecordEventFunc(const ShouldRecordEventFunction& ShouldRecordEventFunc) = 0;
 };
diff --git a/Engine/Source/Runtime/Engine/Private/EngineAnalytics.cpp b/Engine/Source/Runtime/Engine/Private/EngineAnalytics.cpp
--- a/Engine/Source/Runtime/Engine/Private/EngineAnalytics.cpp
+++ b/Engine/Source/Runtime/Engine/Private/EngineAnalytics.cpp
@@ -70,6 +70,10 @@ public:
 	{
 		return Analytics.Pin()->StartSession(MoveTemp(InSessionID), MoveTemp(Attributes));
 	}
+	virtual bool ShouldRecordEvent(const FString& EventName) const override
+	{
+		return Analytics.Pin()->ShouldRecordEvent(EventName);
+	}
 
 	virtual void SetDefaultEventAttributes(TArray<FAnalyticsEventAttribute>&& Attributes) override
 	{
@@ -156,6 +160,10 @@ public:
 
 		Analytics.Pin()->RecordEventJson(EventName, MoveTemp(AttributesJson));
 	}
+	virtual void SetShouldRecordEventFunc(const ShouldRecordEventFunction& ShouldRecordEventFunc)
+	{
+		Analytics.Pin()->SetShouldRecordEventFunc(ShouldRecordEventFunc);
+	}
 
 private:
 	TWeakPtr<IAnalyticsProviderET> Analytics;
