Generalized the Editor analytics summary session system to be usable/extendable by other apps.

Engine/Editor changes:

- Split the Editor summary session in two, one summary for the Engine properties and one for the Editor specific properties. Made it easy to extend the Engine summary to create other summaries.
- Made the summary sender as agnostics as possible of the keys it sends.
- Fixed the system wide lock contention between the process when persisting a session. (On problem caused by the lock is UE-114315).
- Fixed concurrent issue when saving the summary sessions on Linux/Mac
- Fixed performance issue when saving the summary session on Linux/Mac. This enable saving at higher frequency.
- Fixed cases where the same session summary is sent more than once.
- Fixed Windows registry key overflow that could happens if we accumulated too many sessions (in theory, this can happen)
- Made adding new properties to the summary easy and private to the implementation.
- Brought the Linux/Mac implementation closer to Windows implementation.
- Reduced memory allocation, especially when the session records a crash.
- Improved chances to send the summary non-delayed by allowing the Editor to send the reports if CRC died unexpectedly.
- Generalized the support to collect and aggregate analytics from helper processes. For example, CRC already collects analytics that is merged with the Editor summary as information supplement
- Reserved the disk space required to store the summary ahead of time to prevent failing later.
- Increased frequency at which the summary is persisted because saving the summary is more efficient. (About every 10 seconds rather than every minutes).
- Added unit tests

CrashReportClient changes:

- Created a 'session summary' from the CRC point of view to merge with the Editor summary.
- Moved analytics collection in a separated class to make the crash reporting code leaner and less noisy with all the analytics
- Merged the CRC diagnostic logger in the class collecting CRC analytics summary and make the diagnostic log a property in the summary.
- Collected analytics (on behalf of Editor) in a background thread because CRC main thread can be blocked collecting a crash, so it doesn't pay attention to other things
- Added MonitorBatteryLevel and MonitorOnACPower summary properties on Windows. Collected on CRC background thread (never blocked, so we reduce changes to miss the battery running out)
- Added MonitorSessionDuration summary property to track now long CRC ran.
- Added MonitorQuitSignalRecv summary property to detect when CRC is soft killed like: taskkill /PID 1234
- Added MonitorIsReportingCrash summary property to track when CRC dies reporting a crash.
- Added MonitorIsCollectingCrash summary property to track when CRC dies collecting a crash artifacts.
- Added IsProcessingCrash summary property to track when CRC dies processing a crash.
- Added MonitorCrashed summary property to track when CRC exception handler was triggered.
- Added MonitorWasShutdown summary property to track when CRC summary was shutdown
- Added MonitorLoggingOut summary property to track when CRC died because the user was logging out (or as result of shutting down or restarting the computer).
- More accurate value for DeathTimestamp summary property because this is now captured in CRC background thread (which cannot be busy handling a crash)
- Added crash processing timing to CRC diagnostic logs (how long it takes to collect/process a crash).

diff --git a/Engine/Source/Developer/EditorAnalyticsSession/Private/EditorAnalyticsSession.cpp b/Engine/Source/Developer/EditorAnalyticsSession/Private/EditorAnalyticsSession.cpp
--- a/Engine/Source/Developer/EditorAnalyticsSession/Private/EditorAnalyticsSession.cpp
+++ b/Engine/Source/Developer/EditorAnalyticsSession/Private/EditorAnalyticsSession.cpp
@@ -30,7 +30,7 @@ namespace EditorAnalyticsDefs
 	//   Version 1_3 : Added SessionTickCount, UserInteractionCount, IsCrcExeMissing, IsUserLoggingOut, MonitorExitCode and readded lost code to save/load/delete IsLowDriveSpace for 4.26.0.
 	//   Version 1_4 : Added CommandLine, EngineTickCount, LastTickTimestamp, DeathTimestamp and IsDebuggerIgnored for 4.26.0.
 	//   Version 1_5 : Added Stall Detector stats for 5.0.
-	//   Version 1_6 : Added ProcessDiagnostics, Renamed IsDebugger as IsDebuggerPresent
+	//   Version 1_6 : Added ProcessDiagnostics, Renamed IsDebugger as IsDebuggerPresent: This was the last version using this system before the refactor in UE 5.0.
 	static const FString StoreId(TEXT("Epic Games"));
 	static const FString SessionSummaryRoot(TEXT("Unreal Engine/Session Summary"));
 	static const FString DeprecatedVersions[] = { // The session format used by older versions.
@@ -40,8 +40,9 @@ namespace EditorAnalyticsDefs
 		SessionSummaryRoot / TEXT("1_3"),
 		SessionSummaryRoot / TEXT("1_4"),
 		SessionSummaryRoot / TEXT("1_5"),
+		SessionSummaryRoot / TEXT("1_6"),
 	};
-	static const FString SessionSummarySection = SessionSummaryRoot / TEXT("1_6"); // The current session format.
+	static const FString SessionSummarySection = SessionSummaryRoot / TEXT("1_7"); // The current session format.
 	static const FString GlobalLockName(TEXT("UE4_SessionSummary_Lock"));
 	static const FString SessionListStoreKey(TEXT("SessionList"));
 
diff --git a/Engine/Source/Programs/CrashReportClient/CrashReportClientEditor.Build.cs b/Engine/Source/Programs/CrashReportClient/CrashReportClientEditor.Build.cs
--- a/Engine/Source/Programs/CrashReportClient/CrashReportClientEditor.Build.cs
+++ b/Engine/Source/Programs/CrashReportClient/CrashReportClientEditor.Build.cs
@@ -19,6 +19,7 @@ public class CrashReportClientEditor : CrashReportClient
 		PrivateDependencyModuleNames.AddRange(
 			new string[]
 			{
+				"Analytics",
 				"EditorAnalyticsSession",
 			}
 		);
diff --git a/Engine/Source/Programs/CrashReportClient/Private/CrashReportAnalyticsSessionSummary.cpp b/Engine/Source/Programs/CrashReportClient/Private/CrashReportAnalyticsSessionSummary.cpp
new file mode 100644
--- /dev/null
+++ b/Engine/Source/Programs/CrashReportClient/Private/CrashReportAnalyticsSessionSummary.cpp
@@ -0,0 +1,625 @@
+// Copyright Epic Games, Inc. All Rights Reserved.
+
+#include "CrashReportAnalyticsSessionSummary.h"
+#include "AnalyticsPropertyStore.h"
+#include "AnalyticsSessionSummaryManager.h"
+#include "AnalyticsSessionSummarySender.h"
+#include "IAnalyticsProviderET.h"
+#include "CrashReportAnalytics.h"
+#include "Containers/Map.h"
+#include "GenericPlatform/GenericPlatformMath.h"
+#include "GenericPlatform/GenericPlatformCrashContext.h"
+#include "HAL/PlatformTime.h"
+#include "Logging/LogMacros.h"
+#include "Misc/Paths.h"
+#include "Misc/ScopeLock.h"
+#include "Templates/UnrealTemplate.h"
+
+DEFINE_LOG_CATEGORY_STATIC(LogCrashReportClientDiagnostics, Log, All)
+
+#if PLATFORM_WINDOWS && CRASH_REPORT_WITH_MTBF
+
+#include "Windows/AllowWindowsPlatformTypes.h"
+#include "Windows.h"
+
+/** Handle windows messages. */
+LRESULT CALLBACK CrashReportAnalyticsSessionSummaryWindowProc(HWND Hwnd, UINT Msg, WPARAM wParam, LPARAM lParam)
+{
+	// wParam is true if the user session is going away (and CRC is going to die)
+	if (Msg == WM_ENDSESSION && wParam == TRUE)
+	{
+		FCrashReportAnalyticsSessionSummary::Get().OnUserLoggingOut();
+	}
+	else if (Msg == WM_CLOSE)
+	{
+		FCrashReportAnalyticsSessionSummary::Get().OnQuitSignal();
+	}
+	return DefWindowProc(Hwnd, Msg, wParam, lParam);
+}
+
+#endif
+
+namespace CrcAnalyticsProperties
+{
+	/** The exit code of the monitored application. */
+	static const TAnalyticsProperty<int32> MonitoredAppExitCode(TEXT("ExitCode"));
+	/** Track when CRC detected the death of the monitored app. */
+	static const TAnalyticsProperty<FDateTime> MonitoredAppDeathTimestamp(TEXT("DeathTimestamp"));
+
+	/** The CRC startup timestamp. */
+	static const TAnalyticsProperty<FDateTime> StartupTimestamp(TEXT("MonitorStartupTimestamp"));
+	/** The CRC timestamp. */
+	static const TAnalyticsProperty<FDateTime> Timestamp(TEXT("MonitorTimestamp"));
+	/** If CRC raised an exception that was captured by SEH, this is the exception code. */
+	static const TAnalyticsProperty<int32> ExceptCode(TEXT("MonitorExceptCode"));
+	/** If CRC is about to close because the system sent a quit signal. */
+	static const TAnalyticsProperty<bool> QuitSignalRecv(TEXT("MonitorQuitSignalRecv"));
+	/** The CRC diagnostic logs. */
+	static const TAnalyticsProperty<FString> DiagnosticLogs(TEXT("MonitorLogs"));
+	/** The CRC session duration in seconds. */
+	static const TAnalyticsProperty<int32> SessionDurationSecs(TEXT("MonitorSessionDuration"));
+	/** The battery level, if known. */
+	static const TAnalyticsProperty<uint32> BatteryLevel(TEXT("MonitorBatteryLevel"));
+	/** If the system is connected to AC, if know */
+	static const TAnalyticsProperty<bool> IsOnACPower(TEXT("MonitorOnACPower"));
+	/** Whether CRC is reporting a crash. */
+	static const TAnalyticsProperty<bool> IsReportingCrash(TEXT("MonitorIsReportingCrash"));
+	/** Whether CRC is collecting crash artifacts. */
+	static const TAnalyticsProperty<bool> IsCollectingCrash(TEXT("MonitorIsCollectingCrash"));
+	/** Whether CRC is processing a crash. */
+	static const TAnalyticsProperty<bool> IsProcessingCrash(TEXT("MonitorIsProcessingCrash"));
+	/** If CRC is about to be killed because the user is logging out (system shutdown/reboot). */
+	static const TAnalyticsProperty<bool> UserIsLoggingOut(TEXT("MonitorLoggingOut"));
+	/** Whether CRC crashed. */
+	static const TAnalyticsProperty<bool> IsCrashing(TEXT("MonitorCrashed"));
+	/** Whether CRC was shutdown. */
+	static const TAnalyticsProperty<bool> WasShutdown(TEXT("MonitorWasShutdown"));
+}
+
+namespace CrashReportClientUtils
+{
+	/** Flush the store peridically. */
+	static const FTimespan PropertyStoreFlushPeriod = FTimespan::FromSeconds(10);
+
+	/** Tick periodically to poll new information. */
+	static const FTimespan TickPeriod = FTimespan::FromSeconds(0.5);
+
+	/** Maximum length of the diagnostic log. */
+	static constexpr int32 MaxDiagnosticLogLen = 8 * 1024;
+
+#if PLATFORM_WINDOWS && CRASH_REPORT_WITH_MTBF
+	HWND Hwnd = NULL;
+
+	/** Create a hidden Windows to intercept WM_ messages, especially the WM_ENDSESSION. */
+	void InitPlatformSpecific()
+	{
+		// Register the window class.
+		const wchar_t CLASS_NAME[] = L"CRC Analytics Session Window Message Interceptor";
+
+		WNDCLASS wc = { };
+		wc.lpfnWndProc = CrashReportAnalyticsSessionSummaryWindowProc;
+		wc.hInstance = hInstance;
+		wc.lpszClassName = CLASS_NAME;
+
+		RegisterClass(&wc);
+
+		// Create a window to capture WM_ENDSESSION message (so that we can detect when CRC fails because the user is logging off/shutting down/restarting)
+		Hwnd = CreateWindowEx(
+			0,                              // Optional window styles.
+			CLASS_NAME,                     // Window class
+			L"CRC Message Loop Wnd",        // Window text
+			WS_OVERLAPPEDWINDOW,            // Window style
+			CW_USEDEFAULT, CW_USEDEFAULT, CW_USEDEFAULT, CW_USEDEFAULT, // Size and position
+			NULL,         // Parent window
+			NULL,         // Menu
+			hInstance,    // Instance handle
+			NULL          // Additional application data
+		);
+	}
+
+	/** Pump the message from the hidden Windows. */
+	void TickPlatformSpecific()
+	{
+		if (Hwnd != NULL)
+		{
+			// Pump the messages.
+			MSG Msg = { };
+			while (::PeekMessage(&Msg, NULL, 0, 0, PM_REMOVE))
+			{
+				::TranslateMessage(&Msg);
+				::DispatchMessage(&Msg);
+			}
+		}
+	}
+
+	bool GetPowerStatus(TOptional<bool>& OutACPowerConnected, TOptional<uint32>& OutBatteryPct)
+	{
+		bool bAvailable = false;
+		SYSTEM_POWER_STATUS Status;
+		if (GetSystemPowerStatus(&Status))
+		{
+			switch (Status.ACLineStatus)
+			{
+			case 0: // AC Offline
+				OutACPowerConnected.Emplace(false);
+				bAvailable = true;
+				break;
+			case 1: // AC Online
+				OutACPowerConnected.Emplace(true);
+				bAvailable = true;
+				break;
+			default: // Unknown
+				break;
+			}
+
+			if (Status.BatteryLifePercent != 255) // Unknown
+			{
+				OutBatteryPct.Emplace(Status.BatteryLifePercent);
+				bAvailable = true;
+			}
+		}
+		return bAvailable;
+	}
+#else
+	void InitPlatformSpecific(){}
+	void TickPlatformSpecific(){}
+	bool GetPowerStatus(TOptional<bool>& OutACPowerConnected, TOptional<uint32>& OutBatteryPct) { return false; }
+#endif
+
+} // CrashReportClientUtils
+
+
+/**
+ * Augments the default summary senders to perform a short analytis to detect if this is the session ended abnormally. The sender searches
+ * for very specific keys published by the engine. Remember that CRC merges its summary with the monitored process summary (the Editor), so
+ * it gets to see what was recorded (or not) by the monitored process.
+ */
+class FCrashReportClientAnalyticsSessionSummarySender : public FAnalyticsSessionSummarySender
+{
+public:
+	FCrashReportClientAnalyticsSessionSummarySender(IAnalyticsProviderET& Provider)
+		: FAnalyticsSessionSummarySender(Provider)
+	{
+	}
+
+	virtual bool SendSessionSummary(const FString& UserId, const FString& AppId, const FString& AppVersion, const FString& SessionId, const TMap<FString, FString>& Properties) override
+	{
+		// CRC should only send one session (its own), but do a reset in case this convention changes.
+		bAbnormalShutdown = false;
+		bUserLoggingOut = false;
+
+		// Analyze the report to be sent and try to figure out if this is an abnormal shutdown. They keys are taken from the engine analytics session summary. To prevent dependencies
+		// between CRC and Engine analytics, we duplicate the keys here.
+		if (const FString* ShutdownTypeCode = Properties.Find(FAnalyticsSessionSummaryManager::ShutdownTypeCodeProperty.Key))
+		{
+			bAbnormalShutdown = (*ShutdownTypeCode == LexToString((int32)EAnalyticsSessionShutdownType::Abnormal));
+			if (bAbnormalShutdown)
+			{
+				if (const FString* UserLoggingOut = Properties.Find(FAnalyticsSessionSummaryManager::IsUserLoggingOutProperty.Key))
+				{
+					bUserLoggingOut = (*UserLoggingOut == LexToString(true));
+				}
+			}
+		}
+
+		// Send the report unmodified.
+		return FAnalyticsSessionSummarySender::SendSessionSummary(UserId, AppId, AppVersion, SessionId, Properties);
+	}
+
+	/** Returns whether the last session sent was abnormally terminated. */
+	bool IsAbnormalShutdown() const
+	{
+		return bAbnormalShutdown && !bUserLoggingOut;
+	}
+
+private:
+	bool bAbnormalShutdown = false;
+	bool bUserLoggingOut = false;
+};
+
+
+FCrashReportAnalyticsSessionSummary::FCrashReportAnalyticsSessionSummary()
+	: SessionStartTimeSecs(FPlatformTime::Seconds())
+	, bMonitoredAppDeathRecorded(false)
+	, bShutdown(false)
+{
+	if (IsEnabled())
+	{
+		CrashReportClientUtils::InitPlatformSpecific();
+
+		// Reserve the memory for the log string.
+		DiagnosticLog.Reset(CrashReportClientUtils::MaxDiagnosticLogLen);
+		DiagnosticLog.Append(FString::Printf(TEXT("CRC/Init:%s"), *FDateTime::UtcNow().ToString()));
+	}
+}
+
+FCrashReportAnalyticsSessionSummary& FCrashReportAnalyticsSessionSummary::FCrashReportAnalyticsSessionSummary::Get()
+{
+	static FCrashReportAnalyticsSessionSummary Instance;
+	return Instance;
+}
+
+void FCrashReportAnalyticsSessionSummary::Initialize(const FString& ProcessGroupId, uint32 ForProcessId)
+{
+	if (IsEnabled() && !SessionSummaryManager && !ProcessGroupId.IsEmpty())
+	{
+		SessionSummaryManager = MakeUnique<FAnalyticsSessionSummaryManager>(TEXT("CrashReportClient"), ProcessGroupId, ForProcessId);
+		if (SessionSummaryManager)
+		{
+			constexpr uint32 ReservedFileCapacity = CrashReportClientUtils::MaxDiagnosticLogLen + (4 * 1024);
+			PropertyStore = SessionSummaryManager->MakeStore(ReservedFileCapacity);
+			if (PropertyStore)
+			{
+				FCoreDelegates::ApplicationWillTerminateDelegate.AddRaw(this, &FCrashReportAnalyticsSessionSummary::OnApplicationWillTerminate);
+				FCoreDelegates::OnHandleSystemError.AddRaw(this, &FCrashReportAnalyticsSessionSummary::OnHandleSystemError);
+
+				CrcAnalyticsProperties::StartupTimestamp.Set(PropertyStore.Get(), FDateTime::UtcNow());
+				CrcAnalyticsProperties::Timestamp.Set(PropertyStore.Get(), FDateTime::UtcNow());
+				CrcAnalyticsProperties::SessionDurationSecs.Set(PropertyStore.Get(), FMath::FloorToInt(static_cast<float>(FPlatformTime::Seconds() - SessionStartTimeSecs)));
+				CrcAnalyticsProperties::DiagnosticLogs.Set(PropertyStore.Get(), DiagnosticLog, CrashReportClientUtils::MaxDiagnosticLogLen);
+				CrcAnalyticsProperties::IsReportingCrash.Set(PropertyStore.Get(), false);
+				CrcAnalyticsProperties::IsCollectingCrash.Set(PropertyStore.Get(), false);
+				CrcAnalyticsProperties::IsProcessingCrash.Set(PropertyStore.Get(), false);
+				CrcAnalyticsProperties::UserIsLoggingOut.Set(PropertyStore.Get(), false);
+				CrcAnalyticsProperties::IsCrashing.Set(PropertyStore.Get(), false);
+				CrcAnalyticsProperties::WasShutdown.Set(PropertyStore.Get(), false);
+				CrcAnalyticsProperties::QuitSignalRecv.Set(PropertyStore.Get(), false);
+
+				UpdatePowerStatus();
+				Flush();
+
+				GLog->AddOutputDevice(this);
+
+				// CRC main thread might be busy processing a crash, so use a background thread to record important events that could otherwise be missed.
+				AnalyticsThread = MakeUnique<FThread>(TEXT("AnalyticsMonitorThread"), [this, ForProcessId]()
+				{
+					// Try to open the process.
+					FProcHandle MonitoredProcessHandle = FPlatformProcess::OpenProcess(ForProcessId);
+					if (!MonitoredProcessHandle.IsValid())
+					{
+						LogEvent(TEXT("CRC/OpenProcessFailed"));
+						return;
+					}
+					else
+					{
+						LogEvent(TEXT("CRC/Monitoring"));
+					}
+
+					CrashReportClientUtils::InitPlatformSpecific();
+					double NextFlushTimeSecs = FPlatformTime::Seconds();
+					bool bFlushedLowBattery = false;
+
+					while (!bShutdown)
+					{
+						CrashReportClientUtils::TickPlatformSpecific();
+
+						// Monitor the power level.
+						bool bShouldFlush = UpdatePowerStatus();
+
+						if (!FPlatformProcess::IsProcRunning(MonitoredProcessHandle))
+						{
+							OnMonitoredAppDeath(MonitoredProcessHandle);
+							break;
+						}
+						else if (FPlatformTime::Seconds() > NextFlushTimeSecs || bShouldFlush)
+						{
+							// Flush to timestamp the session.
+							Flush();
+							NextFlushTimeSecs += CrashReportClientUtils::PropertyStoreFlushPeriod.GetTotalSeconds();
+						}
+
+						// Throttle the thread.
+						FPlatformProcess::Sleep(CrashReportClientUtils::TickPeriod.GetTotalSeconds());
+					}
+
+					FPlatformProcess::CloseProc(MonitoredProcessHandle);
+				});
+			}
+		}
+	}
+}
+
+bool FCrashReportAnalyticsSessionSummary::IsValid() const
+{
+	return IsEnabled() && PropertyStore.IsValid();
+}
+
+void FCrashReportAnalyticsSessionSummary::Shutdown(IAnalyticsProviderET* AnalyticsProvider, TFunction<void()> HandleAbnormalShutdownFn)
+{
+	if (!IsValid())
+	{
+		return;
+	}
+
+	if (AnalyticsThread)
+	{
+		bShutdown = true;
+		AnalyticsThread->Join();
+		AnalyticsThread.Reset();
+	}
+
+	CrcAnalyticsProperties::WasShutdown.Set(PropertyStore.Get(), true);
+	LogEvent(FString::Printf(TEXT("CRC/Shutdown:%s:%.1fs"), *FDateTime::UtcNow().ToString(), FPlatformTime::Seconds() - SessionStartTimeSecs));
+
+	// Unregister from the core.
+	FCoreDelegates::ApplicationWillTerminateDelegate.RemoveAll(this);
+	FCoreDelegates::OnHandleSystemError.RemoveAll(this);
+	GLog->RemoveOutputDevice(this);
+
+	// Flush and close CRC analytics summary.
+	Flush();
+	PropertyStore.Reset();
+
+	// If CRC is allowed to send summary report on behalf of the monitored application.
+	if (AnalyticsProvider)
+	{
+		// Set a summary sender that will intercept the set of properties and perform an analysis to detect an abnormal shutdown.
+		TSharedPtr<FCrashReportClientAnalyticsSessionSummarySender> SummarySender = MakeShared<FCrashReportClientAnalyticsSessionSummarySender>(*AnalyticsProvider);
+		SessionSummaryManager->SetSender(SummarySender);
+
+		// Merge the principal process summary with CRC process summary and sends the session summary if CRC is the last process to exit.
+		SessionSummaryManager->Shutdown();
+
+		// If the summary sender detected an abnormal shutdown by checking the summary properties, report it.
+		if (SummarySender->IsAbnormalShutdown() && HandleAbnormalShutdownFn)
+		{
+			HandleAbnormalShutdownFn();
+		}
+	}
+	else // Discard CRC summary.
+	{
+		SessionSummaryManager->Shutdown(/*bDiscard*/true);
+	}
+
+	SessionSummaryManager.Reset();
+}
+
+void FCrashReportAnalyticsSessionSummary::Serialize(const TCHAR* V, ELogVerbosity::Type Verbosity, const FName& Category, const double Time)
+{
+	Serialize(V, Verbosity, Category);
+}
+
+void FCrashReportAnalyticsSessionSummary::Serialize(const TCHAR* V, ELogVerbosity::Type Verbosity, const FName& Category)
+{
+	// Log the errors, especially the failed 'check()' with the callstack/message.
+	if (Verbosity == ELogVerbosity::Error || Verbosity == ELogVerbosity::Fatal)
+	{
+		// Log but don't forward to UE logging system. The log is already originate from the logging system.
+		LogEvent(TEXT("CRC/Error"), /*bForwardToUELog*/false);
+		LogEvent(V, /*bForwardToUELog*/false);
+	}
+}
+
+bool FCrashReportAnalyticsSessionSummary::CanBeUsedOnAnyThread() const
+{
+	return true;
+}
+
+bool FCrashReportAnalyticsSessionSummary::CanBeUsedOnMultipleThreads() const
+{
+	return true;
+}
+
+void FCrashReportAnalyticsSessionSummary::LogEvent(const FString& Event)
+{
+	LogEvent(*Event);
+}
+
+void FCrashReportAnalyticsSessionSummary::LogEvent(const TCHAR* Event, bool bForwardToUELog)
+{
+	if (IsValid())
+	{
+		FScopeLock ScopedLock(&LoggerLock);
+		TGuardValue<bool> ReentrantGuard(bLoggerReentrantGuard, true);
+		if (*ReentrantGuard) // Read the old value.
+		{
+			return; // Prevent renentrant logging.
+		}
+
+		AppendLog(Event);
+	}
+
+	// Prevent error logs coming from the logging system to be duplicated.
+	if (bForwardToUELog)
+	{
+		UE_LOG(LogCrashReportClientDiagnostics, Log, TEXT("%s"), Event);
+	}
+}
+
+void FCrashReportAnalyticsSessionSummary::AppendLog(const TCHAR* Event)
+{
+	// Add the separator if some text is already logged.
+	if (DiagnosticLog.Len())
+	{
+		DiagnosticLog.Append(TEXT("|"));
+	}
+
+	// Rotate the log if it gets too long.
+	int32 FreeLen = CrashReportClientUtils::MaxDiagnosticLogLen - DiagnosticLog.Len();
+	int32 EventLen = FCString::Strlen(Event);
+	if (EventLen > FreeLen)
+	{
+		if (EventLen > CrashReportClientUtils::MaxDiagnosticLogLen)
+		{
+			DiagnosticLog.Reset(CrashReportClientUtils::MaxDiagnosticLogLen);
+			EventLen = CrashReportClientUtils::MaxDiagnosticLogLen;
+		}
+		else
+		{
+			DiagnosticLog.RemoveAt(0, EventLen - FreeLen, /*bAllowShrinking*/false); // Free space, remove the chars from the oldest events (in front).
+		}
+	}
+
+	// Append the log entry and dump the log to the file.
+	DiagnosticLog.AppendChars(Event, EventLen);
+
+	// Update the diagnostic field into the session summary store.
+	CrcAnalyticsProperties::DiagnosticLogs.Set(PropertyStore.Get(), DiagnosticLog);
+
+	// Flush the store.
+	Flush();
+}
+
+void FCrashReportAnalyticsSessionSummary::OnMonitoredAppDeath(FProcHandle& Handle)
+{
+	// The first thread to exchange successfull is allowed to update. No need to update twice for the same monitored process.
+	bool Expected = false;
+	if (IsValid() && bMonitoredAppDeathRecorded.compare_exchange_strong(Expected, true))
+	{
+		CrcAnalyticsProperties::MonitoredAppDeathTimestamp.Set(PropertyStore.Get(), FDateTime::UtcNow());
+
+		int32 ExitCode;
+		if (Handle.IsValid() && FPlatformProcess::GetProcReturnCode(Handle, &ExitCode))
+		{
+			CrcAnalyticsProperties::MonitoredAppExitCode.Set(PropertyStore.Get(), ExitCode);
+			LogEvent(FString::Printf(TEXT("App/ExitCode:%d"), ExitCode));
+		}
+		else
+		{
+			CrcAnalyticsProperties::MonitoredAppExitCode.Set(PropertyStore.Get(), ECrashExitCodes::MonitoredApplicationExitCodeNotAvailable);
+			LogEvent(TEXT("App/ExitCode:N/A"));
+		}
+		Flush();
+	}
+}
+
+void FCrashReportAnalyticsSessionSummary::OnUserLoggingOut()
+{
+	if (IsValid())
+	{
+		// The user is logging out and CRC is going to die.
+		CrcAnalyticsProperties::UserIsLoggingOut.Set(PropertyStore.Get(), true);
+
+		// Log the event (this also flush the session).
+		LogEvent(TEXT("CRC/EndSession"));
+	}
+}
+
+void FCrashReportAnalyticsSessionSummary::OnQuitSignal()
+{
+	// The system has requested the app to close. (Like if the user gently kills the application)
+	CrcAnalyticsProperties::QuitSignalRecv.Set(PropertyStore.Get(), true);
+
+	// Log the event (this also flush the session).
+	LogEvent(TEXT("CRC/QuitSignal"));
+}
+
+void FCrashReportAnalyticsSessionSummary::OnCrcCrashing(int32 ExceptCode)
+{
+	if (IsValid())
+	{
+		CrcAnalyticsProperties::IsCrashing.Set(PropertyStore.Get(), true);
+		CrcAnalyticsProperties::ExceptCode.Set(PropertyStore.Get(), ExceptCode);
+	
+		TCHAR CrashEventLog[64];
+		FCString::Sprintf(CrashEventLog, TEXT("CRC/Crash:%d"), ExceptCode);
+		LogEvent(CrashEventLog); // This also flush the session.
+	}
+}
+
+void FCrashReportAnalyticsSessionSummary::Flush()
+{
+	if (IsValid())
+	{
+		// Update the session progression.
+		CrcAnalyticsProperties::Timestamp.Set(PropertyStore.Get(), FDateTime::UtcNow(),
+			[](const FDateTime* Actual, const FDateTime& Proposed) { return Proposed > *Actual; });
+
+		CrcAnalyticsProperties::SessionDurationSecs.Set(PropertyStore.Get(), FMath::FloorToInt(FPlatformTime::Seconds() - SessionStartTimeSecs),
+			[](const int32* Actual, const int32& Proposed) { return Proposed > *Actual; });
+
+		// Flush the store to disk.
+		PropertyStore->Flush();
+	}
+}
+
+void FCrashReportAnalyticsSessionSummary::OnApplicationWillTerminate()
+{
+	LogEvent(FString::Printf(TEXT("CRC/Terminate:%s"), *FDateTime::UtcNow().ToString()));
+}
+
+void FCrashReportAnalyticsSessionSummary::OnHandleSystemError()
+{
+	CrcAnalyticsProperties::IsCrashing.Set(PropertyStore.Get(), true);
+	LogEvent(FString::Printf(TEXT("CRC/SysError:%s"), *FDateTime::UtcNow().ToString()));
+}
+
+void FCrashReportAnalyticsSessionSummary::OnCrashReportStarted(ECrashContextType CrashType, const TCHAR* ErrorMsg)
+{
+	CrashReportStartTimeSecs = FPlatformTime::Seconds();
+
+	CrcAnalyticsProperties::IsReportingCrash.Set(PropertyStore.Get(), true);
+	LogEvent(FString::Printf(TEXT("Report/Start:%s"), *FDateTime::UtcNow().ToString()));
+
+	// Log the assert and ensure condition/file/line/message to the diagnostic log gathered by the analytics to enable searching/grouping them later on.
+	if (CrashType == ECrashContextType::Assert)
+	{
+		LogEvent(FString::Printf(TEXT("Assert/Msg: %s"), ErrorMsg));
+	}
+	else if (CrashType == ECrashContextType::Ensure)
+	{
+		LogEvent(FString::Printf(TEXT("Ensure/Msg: %s"), ErrorMsg));
+	}
+}
+
+void FCrashReportAnalyticsSessionSummary::OnCrashReportCollecting()
+{
+	CrashReportCollectingStartTimeSecs = FPlatformTime::Seconds();
+
+	CrcAnalyticsProperties::IsCollectingCrash.Set(PropertyStore.Get(), true);
+	FCrashReportAnalyticsSessionSummary::Get().LogEvent(TEXT("Report/Collect"));
+}
+
+void FCrashReportAnalyticsSessionSummary::OnCrashReportProcessing(bool bUserInteractive)
+{
+	CrcAnalyticsProperties::IsCollectingCrash.Set(PropertyStore.Get(), false);
+	CrcAnalyticsProperties::IsProcessingCrash.Set(PropertyStore.Get(), true);
+	CrashReportProcessingStartTimeSecs = FPlatformTime::Seconds();
+	FCrashReportAnalyticsSessionSummary::Get().LogEvent(FString::Printf(TEXT("Report/Process:%s"), bUserInteractive ? TEXT("Interactive") : TEXT("Unattended")));
+}
+
+void FCrashReportAnalyticsSessionSummary::OnCrashReportCompleted(bool bSubmitted)
+{
+	CrcAnalyticsProperties::IsCollectingCrash.Set(PropertyStore.Get(), false);
+	CrcAnalyticsProperties::IsProcessingCrash.Set(PropertyStore.Get(), false);
+	CrcAnalyticsProperties::IsReportingCrash.Set(PropertyStore.Get(), false);
+
+	double CurrTimeSecs = FPlatformTime::Seconds();
+	double CollectTimeSecs = CrashReportProcessingStartTimeSecs - CrashReportCollectingStartTimeSecs;
+	double ProcessTimeSecs = CurrTimeSecs - CrashReportProcessingStartTimeSecs;
+	double TotalTimeSecs = CurrTimeSecs - CrashReportStartTimeSecs;
+	const TCHAR* Event = bSubmitted ? TEXT("Report/Sent") : TEXT("Report/Discarded");
+	LogEvent(FString::Printf(TEXT("%s:Collect=%.1fs:Process=%.1fs:Total=%.1fs"), Event, CollectTimeSecs, ProcessTimeSecs, TotalTimeSecs));
+}
+
+bool FCrashReportAnalyticsSessionSummary::UpdatePowerStatus()
+{
+	bool bShouldFlush = false;
+
+	// Monitor the power level.
+	TOptional<bool> ConnectedToACPower;
+	TOptional<uint32> BatteryPercentage;
+	if (CrashReportClientUtils::GetPowerStatus(ConnectedToACPower, BatteryPercentage))
+	{
+		if (ConnectedToACPower)
+		{
+			CrcAnalyticsProperties::IsOnACPower.Set(PropertyStore.Get(), *ConnectedToACPower);
+		}
+		if (BatteryPercentage)
+		{
+			CrcAnalyticsProperties::BatteryLevel.Set(PropertyStore.Get(), *BatteryPercentage, [&bShouldFlush](const uint32* PreviousLevel, const uint32& NewLevel)
+			{
+				// Detects when the batter level goes from above to below or equal a threshonld.
+				constexpr uint32 LowBatteryPct = 2;
+				if (PreviousLevel && *PreviousLevel > LowBatteryPct && NewLevel <= LowBatteryPct)
+				{
+					// Last time, it was above the threshold, but now it dipped below, save the state before the battery runs out.
+					bShouldFlush = true;
+				}
+				return true;
+			});
+		}
+	}
+
+	return bShouldFlush;
+}
diff --git a/Engine/Source/Programs/CrashReportClient/Private/CrashReportAnalyticsSessionSummary.h b/Engine/Source/Programs/CrashReportClient/Private/CrashReportAnalyticsSessionSummary.h
new file mode 100644
--- /dev/null
+++ b/Engine/Source/Programs/CrashReportClient/Private/CrashReportAnalyticsSessionSummary.h
@@ -0,0 +1,134 @@
+// Copyright Epic Games, Inc. All Rights Reserved.
+
+#pragma once
+
+#include "CoreMinimal.h"
+#include "GenericPlatform/GenericPlatformCrashContext.h"
+#include "HAL/PlatformProcess.h"
+#include "Misc/OutputDevice.h"
+#include "Templates/Function.h"
+#include "CrashReportClientDefines.h"
+#include <atomic>
+
+class IAnalyticsPropertyStore;
+class IAnalyticsProviderET;
+class FAnalyticsSessionSummaryManager;
+class FThread;
+
+/**
+ * Creates a summary of CRC session for analytics purpose. The summary should be merged with the
+ * monitored application one.
+ */
+class FCrashReportAnalyticsSessionSummary : public FOutputDevice
+{
+public:
+	/** Return the summary instance. */
+	static FCrashReportAnalyticsSessionSummary& Get();
+
+	/** Initialize the summary. The implementation do nothing unless the instance is successfully initialized. */
+	void Initialize(const FString& ProcessGroupId, uint32 ForProcessId);
+
+	/** Shuts down and submits the summary if a non-null provider is specified and invoke HandleAbnormalShutdownFn if bound and an abnormal shutdown was detected. */
+	void Shutdown(IAnalyticsProviderET* AnalyticsProvider = nullptr, TFunction<void()> HandleAbnormalShutdownFn = TFunction<void()>());
+
+	//~ FOutputDevice interface
+	virtual void Serialize(const TCHAR* V, ELogVerbosity::Type Verbosity, const FName& Category) override;
+	virtual void Serialize(const TCHAR* V, ELogVerbosity::Type Verbosity, const FName& Category, const double Time) override;
+	virtual bool CanBeUsedOnAnyThread() const override;
+	virtual bool CanBeUsedOnMultipleThreads() const override;
+
+	/** Logs an events in the analytics session summary diagnostic log. The event text is expected to be short and concise. */
+	void LogEvent(const TCHAR* Event, bool bForwardToUELog = true);
+
+	/** Logs an events in the analytics session summary diagnostic log. The event text is expected to be short and concise. */
+	void LogEvent(const FString& Event);
+
+	/** Invoked when monitored application death was detected. */
+	void OnMonitoredAppDeath(FProcHandle& Handle);
+
+	/** Invoked by the unhandled exception handler while CRC is crashing. */
+	void OnCrcCrashing(int32 ExceptCode);
+
+	/** Invoked when the user is logging out. This is also called if the computer is shutting down or restarting. */
+	void OnUserLoggingOut();
+
+	/** Invoked when the system gently signal CRC to quit. */
+	void OnQuitSignal();
+
+	/** Invoked when CRC starts to process a crash/ensure/stall.*/
+	void OnCrashReportStarted(ECrashContextType Type, const TCHAR* ErrorMsg);
+
+	/** Invoked when CRC starts collecting the crash artifacts. */
+	void OnCrashReportCollecting();
+
+	/** Invoked when CRC starts the process to submit the crash report (this also includes showing the dialog to user in needed. */
+	void OnCrashReportProcessing(bool bUserInteractive);
+
+	/** Invoked when CRC finished handling a crash report and resume its main loop. */
+	void OnCrashReportCompleted(bool bSubmitted);
+
+private:
+	/** Returns whether the analytics summary is enabled or not. */
+	static bool IsEnabled() { return CRASH_REPORT_WITH_MTBF != 0; }
+
+	/** Returns true if the analytics session summary is valid. */
+	bool IsValid() const;
+
+	/** Creates the summary sessions. */
+	FCrashReportAnalyticsSessionSummary();
+
+	/** Appends a log entry to the log buffer, rotate the buffer if full and flush it to file. */
+	void AppendLog(const TCHAR* Event);
+
+	/** Flush CRC analytics summary file to disk. */
+	void Flush();
+
+	/** Invoked by the core when the application is about to terminate. */
+	void OnApplicationWillTerminate();
+
+	/** Invoked by the core when a system error is detected, like fatal logs. */
+	void OnHandleSystemError();
+
+	/** Invoked periodically by the monitoring thread to update the power status (Battery/AC Power). */
+	bool UpdatePowerStatus();
+
+private:
+	/** Serialize access to the diagnostic logs. */
+	FCriticalSection LoggerLock;
+
+	/** The CRC session summary store. */
+	TSharedPtr<IAnalyticsPropertyStore> PropertyStore;
+
+	/** The session summary manager. */
+	TUniquePtr<FAnalyticsSessionSummaryManager> SessionSummaryManager;
+
+	/** Contains the rotating diagnostic log entries. */
+	FString DiagnosticLog;
+
+	/** Runs in background to collect CRC analytics because the main thread might be busy handling a crash and important event could be missed. */
+	TUniquePtr<FThread> AnalyticsThread;
+
+	/** The CRC startup time, i.e. when the summary is successfully initialized. */
+	double SessionStartTimeSecs;
+
+	/** Time at which the last crash report started. */
+	double CrashReportStartTimeSecs = 0.0;
+
+	/** Time at which collecting the artifacts of the last crash started. */
+	double CrashReportCollectingStartTimeSecs = 0.0;
+
+	/** Time at which processing the last crash started. */
+	double CrashReportProcessingStartTimeSecs = 0.0;
+
+	/** Prevent a reentrency in the diagnostic logger. */
+	bool bLoggerReentrantGuard = false;
+
+	/** Flag if the session is usable. */
+	bool bIsValid = false;
+
+	/** Indicates whether the monitored application death is reported. */
+	std::atomic<bool> bMonitoredAppDeathRecorded;
+
+	/** Used to signal the analytics thread to exit. */
+	std::atomic<bool> bShutdown;
+};
diff --git a/Engine/Source/Programs/CrashReportClient/Private/CrashReportClientApp.cpp b/Engine/Source/Programs/CrashReportClient/Private/CrashReportClientApp.cpp
--- a/Engine/Source/Programs/CrashReportClient/Private/CrashReportClientApp.cpp
+++ b/Engine/Source/Programs/CrashReportClient/Private/CrashReportClientApp.cpp
@@ -25,7 +25,7 @@
 #include "IAnalyticsProviderET.h"
 #include "XmlParser.h"
 #include "Containers/Map.h"
-#include "DiagnosticLogger.h"
+#include "CrashReportAnalyticsSessionSummary.h"
 
 #if !CRASH_REPORT_UNATTENDED_ONLY
 	#include "SCrashReportClient.h"
@@ -47,11 +47,6 @@
 #include "XmlFile.h"
 #include "RecoveryService.h"
 
-#if CRASH_REPORT_WITH_MTBF
-#include "EditorAnalyticsSession.h"
-#include "EditorSessionSummarySender.h"
-#endif
-
 class FRecoveryService;
 
 /** Default main window size */
@@ -92,6 +87,9 @@ static void* MonitorWritePipe = nullptr;
 /** If in monitor mode, set to true when the monitored app crashes. */
 static bool bMonitoredAppCrashed = false;
 
+/** In in monitor mode, this is a strong (much stronger than PID) to uniquely tie the CRC process to the monitored process. */
+FString MonitorProcessGroupId;
+
 /** Result of submission of report */
 enum SubmitCrashReportResult {
 	Failed,				// Failed to send report
@@ -169,11 +167,12 @@ void ParseCommandLine(const TCHAR* CommandLine)
 		MonitorPid = FPlatformString::Atoi64(*Params.FindRef(TEXT("MONITOR")));
 		MonitorReadPipe = (void*) FPlatformString::Atoi64(*Params.FindRef(TEXT("READ")));
 		MonitorWritePipe = (void*) FPlatformString::Atoi64(*Params.FindRef(TEXT("WRITE")));
+		MonitorProcessGroupId = Params.FindRef(TEXT("ProcessGroupId"));
 	}
 
 	if (FoundReportDirectoryAbsolutePaths.Num() == 0)
 	{
-		FPlatformErrorReport::FindMostRecentErrorReports(FoundReportDirectoryAbsolutePaths, FTimespan::FromDays(30));  //FTimespan::FromMinutes(30));
+		FPlatformErrorReport::FindMostRecentErrorReports(FoundReportDirectoryAbsolutePaths, FTimespan::FromDays(30)); //FTimespan::FromMinutes(30));
 	}
 }
 
@@ -304,7 +303,7 @@ SubmitCrashReportResult RunWithUI(FPlatformErrorReport ErrorReport)
 	// Create the main implementation object
 	TSharedRef<FCrashReportClient> CrashReportClient = MakeShared<FCrashReportClient>(ErrorReport);
 
-	// open up the app window	
+	// Open up the app window
 	TSharedRef<SCrashReportClient> ClientControl = SNew(SCrashReportClient, CrashReportClient);
 
 	// Get the engine major version to display in title.
@@ -441,15 +440,15 @@ FPlatformErrorReport CollectErrorReport(FRecoveryService* RecoveryService, uint3
 	FProcHandle ProcessHandle = FPlatformProcess::OpenProcess(Pid);
 	if (!ProcessHandle.IsValid())
 	{
-		FDiagnosticLogger::Get().LogEvent(TEXT("Report/OpenProcessFail"));
+		FCrashReportAnalyticsSessionSummary::Get().LogEvent(TEXT("Report/OpenProcessFail"));
 	}
 	else if (SharedCrashContext.CrashingThreadId == 0)
 	{
-		FDiagnosticLogger::Get().LogEvent(TEXT("Report/BadCrashThreadId"));
+		FCrashReportAnalyticsSessionSummary::Get().LogEvent(TEXT("Report/BadCrashThreadId"));
 	}
 	else if (SharedCrashContext.NumThreads == CR_MAX_THREADS)
 	{
-		FDiagnosticLogger::Get().LogEvent(TEXT("Report/BumpThreadLimits"));
+		FCrashReportAnalyticsSessionSummary::Get().LogEvent(TEXT("Report/BumpThreadLimits"));
 	}
 
 	// First init the static crash context state
@@ -606,7 +605,7 @@ FPlatformErrorReport CollectErrorReport(FRecoveryService* RecoveryService, uint3
 	FPlatformErrorReport ErrorReport(ReportDirectoryAbsolutePath);
 
 	// Link the crash to the Editor summary event to help diagnose the abnormal termination quickly.
-	FDiagnosticLogger::Get().LogEvent(*FPrimaryCrashProperties::Get()->CrashGUID);
+	FCrashReportAnalyticsSessionSummary::Get().LogEvent(*FPrimaryCrashProperties::Get()->CrashGUID);
 
 	// Reset stack walking to allow CRC to implicitly walk its own process and close the monitored process handle.
 	FPlatformStackWalk::InitStackWalkingForProcess(FProcHandle());
@@ -677,7 +676,7 @@ bool IsCrashReportAvailable(uint32 WatchedProcess, FSharedCrashContext& CrashCon
 	// Is data available on the pipe.
 	if (FPlatformProcess::ReadPipeToArray(ReadPipe, Buffer))
 	{
-		FDiagnosticLogger::Get().LogEvent(TEXT("Pipe/Read"));
+		FCrashReportAnalyticsSessionSummary::Get().LogEvent(TEXT("Pipe/Read"));
 
 		// This is to ensure the FSharedCrashContext compiled in the monitored process and this process has the same size.
 		int32 TotalRead = Buffer.Num();
@@ -714,11 +713,11 @@ bool IsCrashReportAvailable(uint32 WatchedProcess, FSharedCrashContext& CrashCon
 
 		if (TotalRead < sizeof(FSharedCrashContext))
 		{
-			FDiagnosticLogger::Get().LogEvent(TEXT("Pipe/NotEnoughData"));
+			FCrashReportAnalyticsSessionSummary::Get().LogEvent(TEXT("Pipe/NotEnoughData"));
 		}
 		else if (TotalRead > sizeof(FSharedCrashContext))
 		{
-			FDiagnosticLogger::Get().LogEvent(TEXT("Pipe/TooMuchData"));
+			FCrashReportAnalyticsSessionSummary::Get().LogEvent(TEXT("Pipe/TooMuchData"));
 		}
 		else
 		{
@@ -726,44 +725,44 @@ bool IsCrashReportAvailable(uint32 WatchedProcess, FSharedCrashContext& CrashCon
 			switch (CrashContext.CrashType)
 			{
 				case ECrashContextType::Assert:
-					FDiagnosticLogger::Get().LogEvent(TEXT("Pipe/Assert"));
+					FCrashReportAnalyticsSessionSummary::Get().LogEvent(TEXT("Pipe/Assert"));
 					bMonitoredAppCrashed = true;
 					break;
 
 				case ECrashContextType::Ensure:
-					FDiagnosticLogger::Get().LogEvent(TEXT("Pipe/Ensure"));
+					FCrashReportAnalyticsSessionSummary::Get().LogEvent(TEXT("Pipe/Ensure"));
 					break;
 
 				case ECrashContextType::Stall:
-					FDiagnosticLogger::Get().LogEvent(TEXT("Pipe/Stall"));
+					FCrashReportAnalyticsSessionSummary::Get().LogEvent(TEXT("Pipe/Stall"));
 					break;
 
 				case ECrashContextType::Crash:
-					FDiagnosticLogger::Get().LogEvent(TEXT("Pipe/Crash"));
+					FCrashReportAnalyticsSessionSummary::Get().LogEvent(TEXT("Pipe/Crash"));
 					bMonitoredAppCrashed = true;
 					break;
 
 				case ECrashContextType::GPUCrash:
-					FDiagnosticLogger::Get().LogEvent(TEXT("Pipe/GPUCrash"));
+					FCrashReportAnalyticsSessionSummary::Get().LogEvent(TEXT("Pipe/GPUCrash"));
 					bMonitoredAppCrashed = true;
 					break;
 
 				case ECrashContextType::Hang:
-					FDiagnosticLogger::Get().LogEvent(TEXT("Pipe/Hang"));
+					FCrashReportAnalyticsSessionSummary::Get().LogEvent(TEXT("Pipe/Hang"));
 					break;
 
 				case ECrashContextType::OutOfMemory:
-					FDiagnosticLogger::Get().LogEvent(TEXT("Pipe/OOM"));
+					FCrashReportAnalyticsSessionSummary::Get().LogEvent(TEXT("Pipe/OOM"));
 					bMonitoredAppCrashed = true;
 					break;
 
 				case ECrashContextType::AbnormalShutdown:
-					FDiagnosticLogger::Get().LogEvent(TEXT("Pipe/AbnormalShutdown"));
+					FCrashReportAnalyticsSessionSummary::Get().LogEvent(TEXT("Pipe/AbnormalShutdown"));
 					bMonitoredAppCrashed = true;
 					break;
 
 				default:
-					FDiagnosticLogger::Get().LogEvent(TEXT("Pipe/Unknown"));
+					FCrashReportAnalyticsSessionSummary::Get().LogEvent(TEXT("Pipe/Unknown"));
 					break;
 			}
 		}
@@ -901,7 +900,7 @@ static void HandleAbnormalShutdown(FSharedCrashContext& CrashContext, uint64 Pro
 		PlatformFile.CopyFile(*LogDestination, CrashContext.UserSettings.LogFilePath);
 
 		// This crash is not a real one, but one to capture the Editor logs in case of abnormal termination.
-		FDiagnosticLogger::Get().LogEvent(TEXT("SyntheticCrash"));
+		FCrashReportAnalyticsSessionSummary::Get().LogEvent(TEXT("SyntheticCrash"));
 
 		bool bPortableCallstackAvailable = false; // Should always be false here. The process is dead.
 		FPlatformErrorReport ErrorReport = CollectErrorReport(RecoveryService.Get(), ProcessID, CrashContext, WritePipe, bPortableCallstackAvailable);
@@ -918,16 +917,6 @@ static void HandleAbnormalShutdown(FSharedCrashContext& CrashContext, uint64 Pro
 	}
 }
 
-static bool WasAbnormalShutdown(const FEditorAnalyticsSession& AnalyticSession)
-{
-	// Check if this was an abnormal shutdown (aka. none of the known shutdown types, and not debugged)
-	return AnalyticSession.bCrashed == false &&
-		AnalyticSession.bGPUCrashed == false &&
-		AnalyticSession.bWasShutdown == false &&
-		AnalyticSession.bIsTerminating == false &&
-		AnalyticSession.bWasEverDebugger == false;
-}
-
 #endif
 
 void RunCrashReportClient(const TCHAR* CommandLine)
@@ -936,14 +925,13 @@ void RunCrashReportClient(const TCHAR* CommandLine)
 	FTaskTagScope ThreadScope(ETaskTag::EGameThread); // Main thread is the game thread.
 #endif
 
-	FDiagnosticLogger::Get().LogEvent(FString::Printf(TEXT("CRC/Init:%s"), *FDateTime::UtcNow().ToString()));
-
 	// Override the stack size for the thread pool.
 	FQueuedThreadPool::OverrideStackSize = 256 * 1024;
 
-	// Initialize the engine. -Messaging enables MessageBus transports required by Concert (Recovery Service).
+	// Initialize the engine.
 	FString FinalCommandLine(CommandLine);
 #if CRASH_REPORT_WITH_RECOVERY
+	// -Messaging enables MessageBus transports required by Concert (Recovery Service).
 	FinalCommandLine += TEXT(" -Messaging -EnablePlugins=\"UdpMessaging,ConcertSyncServer\"");
 #endif
 	GEngineLoop.PreInit(*FinalCommandLine);
@@ -952,10 +940,6 @@ void RunCrashReportClient(const TCHAR* CommandLine)
 	// Increase the HttpSendTimeout to 5 minutes
 	GConfig->SetFloat(TEXT("HTTP"), TEXT("HttpSendTimeout"), 5 * 60.0f, GEngineIni);
 
-	// Pipe UE logs into the diagnostic logger. The diagnostic log is attached to the Editor session summary analytics event (if enabled) and can help diagnose CRC crashes.
-	GLog->AddOutputDevice(&FDiagnosticLogger::Get());
-	FDiagnosticLogger::Get().LogEvent(TEXT("CRC/Load"));
-
 	// Make sure all UObject classes are registered and default properties have been initialized
 	ProcessNewlyLoadedUObjects();
 
@@ -968,8 +952,6 @@ void RunCrashReportClient(const TCHAR* CommandLine)
 	// Load Concert Sync plugins in default phase
 	IPluginManager::Get().LoadModulesForEnabledPlugins(ELoadingPhase::Default);
 
-	FDiagnosticLogger::Get().LogEvent(TEXT("CRC/Config"));
-
 	// Initialize config.
 	FCrashReportCoreConfig::Get();
 
@@ -979,8 +961,6 @@ void RunCrashReportClient(const TCHAR* CommandLine)
 
 	if (MonitorPid == 0) // Does not monitor any process.
 	{
-		FDiagnosticLogger::Get().LogEvent(TEXT("NoMonitor/Start"));
-
 		if (AnalyticsEnabledFromCmd)
 		{
 			FCrashReportAnalytics::Initialize();
@@ -997,29 +977,17 @@ void RunCrashReportClient(const TCHAR* CommandLine)
 	}
 	else // Launched in 'service mode - watches/serves a process'
 	{
-		FDiagnosticLogger::Get().LogEvent(FString::Printf(TEXT("Monitor/Start:%d"), FPlatformProcess::GetCurrentProcessId()));
+		FCrashReportAnalyticsSessionSummary::Get().Initialize(MonitorProcessGroupId, MonitorPid);
 
 		if (!MonitorWritePipe)
 		{
-			FDiagnosticLogger::Get().LogEvent(TEXT("CRC/NoWritePipe"));
+			FCrashReportAnalyticsSessionSummary::Get().LogEvent(TEXT("CRC/NoWritePipe"));
 		}
 		if (!MonitorReadPipe)
 		{
-			FDiagnosticLogger::Get().LogEvent(TEXT("CRC/NoReadPipe"));
+			FCrashReportAnalyticsSessionSummary::Get().LogEvent(TEXT("CRC/NoReadPipe"));
 		}
 
-		// Log any termination occurring in CRC (no sub-system should terminate CRC)
-		FDelegateHandle TerminateHandle = FCoreDelegates::ApplicationWillTerminateDelegate.AddLambda([]()
-		{
-			FDiagnosticLogger::Get().LogEvent(FString::Printf(TEXT("CRC/Terminate:%s"), *FDateTime::UtcNow().ToString()));
-		});
-
-		// Log any system error occurring in CRC.
-		FDelegateHandle SystemErrorHandle = FCoreDelegates::OnHandleSystemError.AddLambda([]()
-		{
-			FDiagnosticLogger::Get().LogEvent(FString::Printf(TEXT("CRC/SysError:%s"), *FDateTime::UtcNow().ToString()));
-		});
-
 		const int32 IdealFramerate = 10;
 		double PrevLoopStartTime = FPlatformTime::Seconds();
 		const float IdealFrameTime = 1.0f / IdealFramerate;
@@ -1029,7 +997,7 @@ void RunCrashReportClient(const TCHAR* CommandLine)
 #if CRASH_REPORT_WITH_RECOVERY
 		// Starts the disaster recovery service. This records transactions and allows users to recover from previous crashes.
 		RecoveryServicePtr = MakeShared<FRecoveryService>(MonitorPid);
-		FDiagnosticLogger::Get().LogEvent(TEXT("Recovery/Started"));
+		FCrashReportAnalyticsSessionSummary::Get().LogEvent(TEXT("Recovery/Started"));
 #endif
 
 		// Try to open the process.
@@ -1039,41 +1007,8 @@ void RunCrashReportClient(const TCHAR* CommandLine)
 #else
 		FProcHandle MonitoredProcess = FPlatformProcess::OpenProcess(MonitorPid);
 #endif
-
-		FDiagnosticLogger::Get().LogEvent(MonitoredProcess.IsValid() ? TEXT("OpenProcess/Done") : TEXT("OpenProcess/Failed"));
-
-		auto GetProcessStatus = [](FProcHandle& ProcessHandle) -> TTuple<bool/*Running*/, TOptional<int32>/*ReturnCode*/>
-		{
-			bool bRunning = true;
-			TOptional<int32> ProcessReturnCodeOpt; // Unknown by default.
-			if (!ProcessHandle.IsValid())
-			{
-				bRunning = false;
-			}
-			else if (!FPlatformProcess::IsProcRunning(ProcessHandle))
-			{
-				bRunning = false;
-				int32 ProcessReturnCode = 0;
-				if (FPlatformProcess::GetProcReturnCode(ProcessHandle, &ProcessReturnCode)) // Is the return code available? (Not supported on all platforms)
-				{
-					FDiagnosticLogger::Get().LogEvent(FString::Printf(TEXT("Editor/ExitCode:%d"), ProcessReturnCode));
-					ProcessReturnCodeOpt.Emplace(ProcessReturnCode);
-				}
-				else
-				{
-					FDiagnosticLogger::Get().LogEvent(TEXT("Editor/ExitCode:N/A"));
-				}
-			}
-
-			return MakeTuple(bRunning, ProcessReturnCodeOpt);
-		};
-
-		// The approximative time of death of the monitored application.
-		FDateTime MonitoredProcessDeathTime;
-
 		// Loop until the monitored process dies.
-		TTuple<bool/*bRunning*/, TOptional<int32>/*ExitCode*/> ProcessStatus = GetProcessStatus(MonitoredProcess);
-		while (ProcessStatus.Get<0>())
+		while (MonitoredProcess.IsValid() && FPlatformProcess::IsProcRunning(MonitoredProcess))
 		{
 			const double CurrLoopStartTime = FPlatformTime::Seconds();
 
@@ -1084,7 +1019,7 @@ void RunCrashReportClient(const TCHAR* CommandLine)
 				FSharedCrashContext CrashContext;
 				if (IsCrashReportAvailable(MonitorPid, CrashContext, MonitorReadPipe) && !IsEngineExitRequested())
 				{
-					FDiagnosticLogger::Get().LogEvent(TEXT("Report/Start"));
+					FCrashReportAnalyticsSessionSummary::Get().OnCrashReportStarted(CrashContext.CrashType, CrashContext.ErrorMessage);
 
 					const bool bReportCrashAnalyticInfo = CrashContext.UserSettings.bSendUsageData;
 					if (bReportCrashAnalyticInfo)
@@ -1092,7 +1027,7 @@ void RunCrashReportClient(const TCHAR* CommandLine)
 						FCrashReportAnalytics::Initialize();
 					}
 
-					FDiagnosticLogger::Get().LogEvent(TEXT("Report/Collect"));
+					FCrashReportAnalyticsSessionSummary::Get().OnCrashReportCollecting();
 
 					// Build error report in memory.
 					bool bCrashedThreadCallstackAvailable = false;
@@ -1102,35 +1037,21 @@ void RunCrashReportClient(const TCHAR* CommandLine)
 					// that the timeout in the crashing process waiting for CRC to reply 'continue' is too short.
 					if (!bCrashedThreadCallstackAvailable)
 					{
-						FDiagnosticLogger::Get().LogEvent(TEXT("Report/NoPCallstack"));
+						FCrashReportAnalyticsSessionSummary::Get().LogEvent(TEXT("Report/NoPCallstack"));
 					}
 
 #if CRASH_REPORT_WITH_RECOVERY
 					if (RecoveryServicePtr && !FPrimaryCrashProperties::Get()->bIsEnsure)
 					{
 						// Shutdown the recovery service. This will releases the recovery database file lock (not sharable) and let a new instance take it and offer the user to recover.
-						FDiagnosticLogger::Get().LogEvent(TEXT("Recovery/Shutdown"));
+						FCrashReportAnalyticsSessionSummary::Get().LogEvent(TEXT("Recovery/Shutdown"));
 						RecoveryServicePtr.Reset();
 					}
 #endif
-					FDiagnosticLogger::Get().LogEvent(TEXT("Report/Sending"));
-
 					const bool bNoDialog = (CrashContext.UserSettings.bNoDialog || CrashContext.UserSettings.bImplicitSend) && CrashContext.UserSettings.bSendUnattendedBugReports;
+					FCrashReportAnalyticsSessionSummary::Get().OnCrashReportProcessing(/*bIsUserInteractive*/!bNoDialog);
 					const SubmitCrashReportResult Result = SendErrorReport(ErrorReport, bNoDialog, CrashContext.UserSettings.bImplicitSend);
 
-					FDiagnosticLogger::Get().LogEvent(Result == SubmitCrashReportResult::SuccessDiscarded ? TEXT("Report/Discarded") :
-					                                 (Result == SubmitCrashReportResult::Failed ? TEXT("Report/Failed") : TEXT("Report/Sent")));
-
-					// Log the assert and ensure condition/file/line/message to the diagnostic log gathered by the analytics to enable searching/grouping them later on.
-					if (CrashContext.CrashType == ECrashContextType::Assert)
-					{
-						FDiagnosticLogger::Get().LogEvent(FString::Printf(TEXT("Assert/Msg: %s"), CrashContext.ErrorMessage));
-					}
-					else if (CrashContext.CrashType == ECrashContextType::Ensure)
-					{
-						FDiagnosticLogger::Get().LogEvent(FString::Printf(TEXT("Ensure/Msg: %s"), CrashContext.ErrorMessage));
-					}
-
 					if (bReportCrashAnalyticInfo)
 					{
 						if (FCrashReportAnalytics::IsAvailable())
@@ -1141,7 +1062,7 @@ void RunCrashReportClient(const TCHAR* CommandLine)
 						FCrashReportAnalytics::Shutdown();
 					}
 
-					FDiagnosticLogger::Get().LogEvent(TEXT("Report/Done"));
+					FCrashReportAnalyticsSessionSummary::Get().OnCrashReportCompleted(Result != SubmitCrashReportResult::SuccessDiscarded && Result != SubmitCrashReportResult::Failed);
 				}
 			}
 
@@ -1160,131 +1081,50 @@ void RunCrashReportClient(const TCHAR* CommandLine)
 			IncrementalPurgeGarbage(true, FMath::Clamp(PurgeSeconds, 0.002f, 1.0f)));
 #endif
 			GLog->FlushThreadedLogs();
-
+			
 			// Throttle main thread fps by sleeping if we still have time.
 			const float SleepSeconds = IdealFrameTime - static_cast<float>(FPlatformTime::Seconds() - CurrLoopStartTime);
 			FPlatformProcess::Sleep(FMath::Clamp(SleepSeconds, 0.0f, 1.0f));
-
-			// Refresh the monitored application status.
-			ProcessStatus = GetProcessStatus(MonitoredProcess);
-			if (!ProcessStatus.Get<0>()) // Not running anymore.
-			{
-				MonitoredProcessDeathTime = FDateTime::UtcNow();
-			}
-
 			PrevLoopStartTime = CurrLoopStartTime;
-
-			// Tick the logger so that it periodically timestamp the mini-log file to detect approximatively when CRC exited (or hang).
-			FDiagnosticLogger::Get().Tick();
 		}
 
-#if CRASH_REPORT_WITH_MTBF // Expected to be 1 when compiling CrashReportClientEditor.
-		{
-			FDiagnosticLogger::Get().LogEvent(FString::Printf(TEXT("MTBF/Start:%s"), *FDateTime::UtcNow().ToString()));
-
-			// Get the status of the Editor process.
-			TOptional<int32> MonitoredProcessExitCode = ProcessStatus.Get<1>();
-			bool bMonitoredProcessExited = !ProcessStatus.Get<0>();
-			bool bMonitoredSessionLoaded = false;
-
-			FDiagnosticLogger::Get().LogEvent(TEXT("MTBF/LoadSession"));
-
-			// Try to persist an exit code in session summary.
-			FEditorAnalyticsSession MonitoredSession;
-			double WaitEndTime = FPlatformTime::Seconds() + 60;
-			FTimespan LockTimeout = FTimespan::FromMilliseconds(5);
-			bool bSessionLockAcquired = false;
-			while (FPlatformTime::Seconds() <= WaitEndTime)
-			{
-				if (FEditorAnalyticsSession::Lock(LockTimeout)) // Don't block for a long time to keep ticking the diagnostic logger regularly.
-				{
-					bSessionLockAcquired = true;
-					if (FEditorAnalyticsSession::FindSession(MonitorPid, MonitoredSession))
-					{
-						bMonitoredSessionLoaded = true;
-						if (!MonitoredSession.SaveExitCode(MonitoredProcessExitCode.IsSet() ? MonitoredProcessExitCode.GetValue() : ECrashExitCodes::MonitoredApplicationExitCodeNotAvailable, MonitoredProcessDeathTime))
-						{
-							FDiagnosticLogger::Get().LogEvent(TEXT("MTBF/ExitCodeNotSaved"));
-						}
-					}
-					else
-					{
-						FDiagnosticLogger::Get().LogEvent(TEXT("MTBF/NoSessionFound"));
-						if (bMonitoredAppCrashed || (MonitoredProcessExitCode.IsSet() && *MonitoredProcessExitCode != 0))
-						{
-							// No session found. Record this crash to piggyback off a valid session later. (Purpose: Account for early Editor crash happening before Analytics was initialized)
-							FEditorAnalyticsSession::CreateMinimalCrashSession(MonitoredProcessExitCode);
-						}
-					}
-					FEditorAnalyticsSession::Unlock();
-					break;
-				}
+		FCrashReportAnalyticsSessionSummary::Get().OnMonitoredAppDeath(MonitoredProcess);
 
-				// Periodically timestamp the mini-log and check if the application is about to die (user logging off/shutting down the computer).
-				FDiagnosticLogger::Get().Tick();
-			}
-
-			if (!bSessionLockAcquired) // Too much contention on the session lock (or the lock is corrupted).
-			{
-				FDiagnosticLogger::Get().LogEvent(TEXT("MTBF/LockSessionFail"));
-			}
-
-			if (bMonitoredProcessExited)
+#if CRASH_REPORT_WITH_MTBF
+		{
+			// Load the temporary crash context file.
+			FSharedCrashContext TempCrashContext;
+			FMemory::Memzero(TempCrashContext);
+			if (LoadTempCrashContextFromFile(TempCrashContext, MonitorPid) && TempCrashContext.UserSettings.bSendUsageData)
 			{
-				// Load the temporary crash context file.
-				FSharedCrashContext TempCrashContext;
-				FMemory::Memzero(TempCrashContext);
-				if (LoadTempCrashContextFromFile(TempCrashContext, MonitorPid) && TempCrashContext.UserSettings.bSendUsageData)
+				FCrashReportAnalytics::Initialize();
+				if (FCrashReportAnalytics::IsAvailable())
 				{
-					FCrashReportAnalytics::Initialize();
-					if (FCrashReportAnalytics::IsAvailable())
+					auto HandleAbnormalShutdownFunc = [&TempCrashContext, &RecoveryServicePtr]()
 					{
-						FDiagnosticLogger::Get().LogEvent(TEXT("MTBF/Done"));
-
-						// If the Editor thinks the session ended up abnormally, generate a crash report (to get the Editor logs and figure out why this happened).
-						if (bMonitoredSessionLoaded && TempCrashContext.UserSettings.bSendUnattendedBugReports)
+						if (TempCrashContext.UserSettings.bSendUnattendedBugReports)
 						{
-							// Check what the Editor knows about the exit. Was the proper handlers called and the flag(s) set in the summary event?
-							if (WasAbnormalShutdown(MonitoredSession) && !MonitoredSession.bIsUserLoggingOut)
-							{
-								// Send a spoofed crash report in the case that we detect an abnormal shutdown has occurred
-								HandleAbnormalShutdown(TempCrashContext, MonitorPid, MonitorWritePipe, RecoveryServicePtr);
-							}
+							// Send a spoofed crash report in the case that we detect an abnormal shutdown has occurred
+							HandleAbnormalShutdown(TempCrashContext, MonitorPid, MonitorWritePipe, RecoveryServicePtr);
 						}
+					};
 
-						// Stop the logging so that LoadAllLogs() and ClearAllLogs() can access this mini-log file.
-						FDiagnosticLogger::Get().Close();
-
-						// Send this session summary event (and the orphan ones if any).
-						FEditorSessionSummarySender EditorSessionSummarySender(FCrashReportAnalytics::GetProvider(), TEXT("CrashReportClient"), MonitorPid);
-						EditorSessionSummarySender.SetMonitorDiagnosticLogs(FDiagnosticLogger::LoadAllLogs());
-						EditorSessionSummarySender.Shutdown();
-						FDiagnosticLogger::ClearAllLogs(); // Logs (if any) were attached and sent if required and are not longer required.
-					}
-					FCrashReportAnalytics::Shutdown();
-				}
-				else
-				{
-					FDiagnosticLogger::Get().LogEvent(TEXT("MTBF/NoTempCrash"));
+					// Shutdown the session, sends the summary and if the session ended up abnormally (analyzing the summary), invoke the functor to spoof a crash report.
+					FCrashReportAnalyticsSessionSummary::Get().Shutdown(&FCrashReportAnalytics::GetProvider(), HandleAbnormalShutdownFunc);
 				}
-			}
-			else
-			{
-				FDiagnosticLogger::Get().LogEvent(TEXT("MTBF/StillRunning"));
+				FCrashReportAnalytics::Shutdown();
 			}
 		}
 #endif
+		// Ensure to shutdown the summary analytics. If it was already shutdown above, this do nothing, otherwise, it destroys analytics data gathered by CRC.
+		FCrashReportAnalyticsSessionSummary::Get().Shutdown();
+
 		// Clean up the context file
 		DeleteTempCrashContextFile(MonitorPid);
 
 		FPlatformProcess::CloseProc(MonitoredProcess);
-
-		FCoreDelegates::ApplicationWillTerminateDelegate.Remove(TerminateHandle);
-		FCoreDelegates::OnHandleSystemError.Remove(SystemErrorHandle);
 	}
 
-	GLog->RemoveOutputDevice(&FDiagnosticLogger::Get());
-
 	FPrimaryCrashProperties::Shutdown();
 	FPlatformErrorReport::ShutDown();
 
@@ -1295,7 +1135,7 @@ void RunCrashReportClient(const TCHAR* CommandLine)
 
 	FEngineLoop::AppPreExit();
 	FModuleManager::Get().UnloadModulesAtShutdown();
-	FTaskGraphInterface::Shutdown();
+	//FTaskGraphInterface::Shutdown();
 
 	FEngineLoop::AppExit();
 }
diff --git a/Engine/Source/Programs/CrashReportClient/Private/DiagnosticLogger.cpp b/Engine/Source/Programs/CrashReportClient/Private/DiagnosticLogger.cpp
deleted file mode 100644
diff --git a/Engine/Source/Programs/CrashReportClient/Private/DiagnosticLogger.h b/Engine/Source/Programs/CrashReportClient/Private/DiagnosticLogger.h
deleted file mode 100644
diff --git a/Engine/Source/Programs/CrashReportClient/Private/Windows/CrashReportClientMainWindows.cpp b/Engine/Source/Programs/CrashReportClient/Private/Windows/CrashReportClientMainWindows.cpp
--- a/Engine/Source/Programs/CrashReportClient/Private/Windows/CrashReportClientMainWindows.cpp
+++ b/Engine/Source/Programs/CrashReportClient/Private/Windows/CrashReportClientMainWindows.cpp
@@ -5,13 +5,12 @@
 #include "CrashReportClientDefines.h"
 
 #if CRASH_REPORT_WITH_MTBF
-#include "EditorAnalyticsSession.h"
 #include "HAL/FileManager.h"
 #include "HAL/PlatformProcess.h"
 #include "HAL/PlatformAtomics.h"
 #include "HAL/PlatformStackWalk.h"
 #include "Serialization/Archive.h"
-#include "DiagnosticLogger.h"
+#include "CrashReportAnalyticsSessionSummary.h"
 #endif
 
 #if CRASH_REPORT_WITH_MTBF && !PLATFORM_SEH_EXCEPTIONS_DISABLED
@@ -25,27 +24,7 @@ void SaveCrcCrashException(EXCEPTION_POINTERS* ExceptionInfo)
 	static volatile int32 CrashCount = 0;
 	if (FPlatformAtomics::InterlockedIncrement(&CrashCount) == 1)
 	{
-		TCHAR CrashEventLog[64];
-		FCString::Sprintf(CrashEventLog, TEXT("CRC/Crash:%d"), ExceptionInfo->ExceptionRecord->ExceptionCode);
-		FDiagnosticLogger::Get().LogEvent(CrashEventLog);
-
-		uint64 MonitoredEditorPid;
-		if (FParse::Value(GetCommandLineW(), TEXT("-MONITOR="), MonitoredEditorPid))
-		{
-			FTimespan Timeout = FTimespan::FromSeconds(2);
-			if (FEditorAnalyticsSession::Lock(Timeout)) // This lock is reentrant for the same process.
-			{
-				FEditorAnalyticsSession MonitoredSession;
-				if (FEditorAnalyticsSession::FindSession(MonitoredEditorPid, MonitoredSession))
-				{
-					if (!MonitoredSession.SaveMonitorExceptCode(ExceptionInfo->ExceptionRecord->ExceptionCode))
-					{
-						FDiagnosticLogger::Get().LogEvent("CRC/ExceptCodeNotSaved");
-					}
-				}
-				FEditorAnalyticsSession::Unlock();
-			}
-		}
+		FCrashReportAnalyticsSessionSummary::Get().OnCrcCrashing(ExceptionInfo->ExceptionRecord->ExceptionCode);
 
 		if (ExceptionInfo->ExceptionRecord->ExceptionCode != STATUS_HEAP_CORRUPTION)
 		{
@@ -56,7 +35,7 @@ void SaveCrcCrashException(EXCEPTION_POINTERS* ExceptionInfo)
 				FPlatformStackWalk::StackWalkAndDump(CrashStackTrace, UE_ARRAY_COUNT(CrashStackTrace), 0);
 				if (CrashStackTrace[0] != 0)
 				{
-					FDiagnosticLogger::Get().LogEvent(ANSI_TO_TCHAR(CrashStackTrace));
+					FCrashReportAnalyticsSessionSummary::Get().LogEvent(ANSI_TO_TCHAR(CrashStackTrace));
 				}
 			}
 		}
diff --git a/Engine/Source/Runtime/Analytics/Analytics/Private/AnalyticsPropertyStore.cpp b/Engine/Source/Runtime/Analytics/Analytics/Private/AnalyticsPropertyStore.cpp
new file mode 100644
--- /dev/null
+++ b/Engine/Source/Runtime/Analytics/Analytics/Private/AnalyticsPropertyStore.cpp
@@ -0,0 +1,1011 @@
+// Copyright Epic Games, Inc. All Rights Reserved.
+
+#include "AnalyticsPropertyStore.h"
+#include "Containers/StringConv.h"
+#include "GenericPlatform/GenericPlatformFile.h"
+#include "HAL/FileManager.h"
+#include "HAL/PlatformFileManager.h"
+#include "HAL/PlatformTime.h"
+#include "Misc/SecureHash.h"
+#include "Misc/ScopeLock.h"
+#include "ProfilingDebugging/CpuProfilerTrace.h"
+
+namespace AnalyticsPropertyStoreUtils
+{
+	/** The key/value storage format version defininig the memory layout of the store. */
+	static const uint32 StoreFormatVersion = 1;
+
+	/** MD5 checksum is 16 bytes long. */
+	static constexpr uint32 Md5DigestSize = 16;
+
+	/** The store data header that should never change. */
+	struct FAnalyticsStoreFileHeader
+	{
+		uint32 Version;
+		uint32 DataSize;
+		uint8 Checksum[Md5DigestSize];
+	};
+
+	/** Conditionnaly set a value is the store. */
+	template<typename T, typename TGetter, typename TSetter, typename TCompare>
+	IAnalyticsPropertyStore::EStatusCode SetValue(TGetter&& GetterFn, TSetter&& SetterFn, const T& ProposedValue, TCompare&& ConditionFn)
+	{
+		T ActualValue;
+		IAnalyticsPropertyStore::EStatusCode Status = GetterFn(ActualValue);
+		if (Status == IAnalyticsPropertyStore::EStatusCode::Success || Status == IAnalyticsPropertyStore::EStatusCode::NotFound)
+		{
+			// Apply the condition.
+			if (ConditionFn(Status == IAnalyticsPropertyStore::EStatusCode::NotFound ? nullptr : &ActualValue, ProposedValue))
+			{
+				return SetterFn(ProposedValue);
+			}
+			else
+			{
+				return IAnalyticsPropertyStore::EStatusCode::Declined;
+			}
+		}
+		return Status;
+	}
+
+	/** Conditionnaly update a value is the store. */
+	template<typename T, typename TGetter, typename TSetter, typename TCompare>
+	IAnalyticsPropertyStore::EStatusCode UpdateValue(TGetter&& GetterFn, TSetter&& SetterFn, TCompare&& UpdateFn)
+	{
+		T Value;
+		IAnalyticsPropertyStore::EStatusCode Status = GetterFn(Value);
+		if (Status == IAnalyticsPropertyStore::EStatusCode::Success)
+		{
+			// Call back and check if the value was updated
+			if (UpdateFn(Value))
+			{
+				return SetterFn(Value);
+			}
+			else
+			{
+				return IAnalyticsPropertyStore::EStatusCode::Declined;
+			}
+		}
+		return Status;
+	}
+} // namespace AnalyticsPropertyStoreUtils
+
+
+FAnalyticsPropertyStore::FAnalyticsPropertyStore()
+	: StorageWriter(StorageBuf)
+	, StorageReader(StorageBuf)
+{
+}
+
+FAnalyticsPropertyStore::~FAnalyticsPropertyStore()
+{
+	FScopeLock Lock(&StoreLock);
+
+	// Ensure any running async flush has completed.
+	if (FlushTask.IsValid())
+	{
+		FlushTask->EnsureCompletion();
+		FlushTask.Reset();
+	}
+}
+
+void FAnalyticsPropertyStore::Reset()
+{
+	// Ensure any running async flush has completed.
+	if (FlushTask.IsValid())
+	{
+		FlushTask->EnsureCompletion();
+		FlushTask.Reset();
+	}
+
+	FileHandle.Reset();
+	NameOffsetMap.Reset();
+	StorageBuf.Empty(StorageBuf.Num());
+	StorageReader.Seek(0);
+	StorageWriter.Seek(0);
+}
+
+bool FAnalyticsPropertyStore::Create(const FString& Pathname, uint32 CapacityHint)
+{
+	FScopeLock Lock(&StoreLock);
+
+	// Reset internal state.
+	Reset();
+
+	CapacityHint = FMath::Max(static_cast<uint32>(sizeof(AnalyticsPropertyStoreUtils::FAnalyticsStoreFileHeader)), CapacityHint);
+	uint32 ReservedSize = FMath::Min(CapacityHint, 2u * 1024 * 1024); // Limit the reserved space to 2MB.
+	StorageBuf.Reserve(ReservedSize - sizeof(AnalyticsPropertyStoreUtils::FAnalyticsStoreFileHeader));
+
+	// Open the file in read/write mode (random access needed).
+	FileHandle = TUniquePtr<IFileHandle>(FPlatformFileManager::Get().GetPlatformFile().OpenWrite(*Pathname, /*bAppend*/false, /*bAllowRead*/false));
+	if (!FileHandle)
+	{
+		return false;
+	}
+
+	// Reserve the space in the file. Nothing else until the store is flushed the first time.
+	uint8 Dummy = 0;
+	FileHandle->Seek(ReservedSize - sizeof(Dummy));
+	FileHandle->Write(&Dummy, sizeof(Dummy));
+	return true;
+}
+
+bool FAnalyticsPropertyStore::Load(const FString& Pathname)
+{
+	FScopeLock Lock(&StoreLock);
+
+	// Reset internal state.
+	Reset();
+
+	// Open the file in read/write mode (random access needed).
+	FileHandle = TUniquePtr<IFileHandle>(FPlatformFileManager::Get().GetPlatformFile().OpenWrite(*Pathname, /*bAppend*/true, /*bAllowRead*/true));
+	if (!FileHandle)
+	{
+		return false;
+	}
+
+	// Ensure the file header data is there. [Version | Checksum | StoreDataSize]. The file can be larger that the stored data size.
+	int64 FileSize = FileHandle->Size();
+	if (FileSize < sizeof(AnalyticsPropertyStoreUtils::FAnalyticsStoreFileHeader))
+	{
+		FileHandle.Reset();
+		return false;
+	}
+
+	// Read the header.
+	FileHandle->Seek(0);
+	AnalyticsPropertyStoreUtils::FAnalyticsStoreFileHeader FileHeader;
+	FileHandle->Read(reinterpret_cast<uint8*>(&FileHeader), sizeof(FileHeader));
+
+	// If the file format is incompatible.
+	if (FileHeader.Version != AnalyticsPropertyStoreUtils::StoreFormatVersion)
+	{
+		FileHandle.Reset();
+		return false;
+	}
+
+	// Read the stored data.
+	StorageBuf.Empty(FileHeader.DataSize);
+	StorageBuf.AddUninitialized(FileHeader.DataSize);
+	FileHandle->Read(StorageBuf.GetData(), FileHeader.DataSize);
+
+	// Compute the checksum from the stored data.
+	uint8 Checksum[AnalyticsPropertyStoreUtils::Md5DigestSize];
+	FMD5 Md5;
+	Md5.Update(StorageBuf.GetData(), StorageBuf.Num());
+	Md5.Final(Checksum);
+
+	// If the checksums don't match.
+	if (FMemory::Memcmp(FileHeader.Checksum, Checksum, AnalyticsPropertyStoreUtils::Md5DigestSize) != 0)
+	{
+		FileHandle.Reset();
+		return false;
+	}
+
+	// Rebuild the key/offset map, reading the stream sequentially.
+	FString Key;
+	uint32 Offset;
+	while (StorageReader.Tell() < StorageReader.TotalSize())
+	{
+		// Get the current offset.
+		Offset = static_cast<uint32>(StorageReader.Tell());
+
+		// Read the type.
+		ETypeCode TypeCode;
+		StorageReader.Serialize(&TypeCode, sizeof(ETypeCode));
+
+		switch (RawType(TypeCode))
+		{
+			case ETypeCode::I32:
+			case ETypeCode::U32:
+			case ETypeCode::Flt:
+			{
+				static_assert(sizeof(uint32) == sizeof(int32) && sizeof(uint32) == sizeof(float), "Incompatible data size");
+				StorageReader.Seek(StorageReader.Tell() + sizeof(uint32)); // Skip over the value, it is not needed.
+				StorageReader << Key;
+				if (!IsDead(TypeCode))
+				{
+					NameOffsetMap.Emplace(Key, Offset);
+				}
+				break;
+			}
+
+			case ETypeCode::I64:
+			case ETypeCode::U64:
+			case ETypeCode::Dbl:
+			case ETypeCode::Date:
+			{
+				static_assert(sizeof(uint64) == sizeof(int64) && sizeof(uint64) == sizeof(double) && sizeof(uint64) == sizeof(decltype(static_cast<FDateTime*>(nullptr)->GetTicks())), "Incompatible data size");
+				StorageReader.Seek(StorageReader.Tell() + sizeof(uint64)); // Skip over the value, it is not needed.
+				StorageReader << Key;
+				if (!IsDead(TypeCode))
+				{
+					NameOffsetMap.Emplace(Key, Offset);
+				}
+				break;
+			}
+
+			case ETypeCode::Bool:
+			{
+				StorageReader.Seek(StorageReader.Tell() + sizeof(bool)); // Skip over the value, it is not needed.
+				StorageReader << Key;
+				if (!IsDead(TypeCode))
+				{
+					NameOffsetMap.Emplace(Key, Offset);
+				}
+				break;
+			}
+
+			case ETypeCode::Str:
+			{
+				uint32 CapacityInBytes;
+				StorageReader.Seek(StorageReader.Tell() + sizeof(uint32)); // Skip over the size
+				StorageReader.Serialize(&CapacityInBytes, sizeof(CapacityInBytes)); // Read the capacity.
+				StorageReader.Seek(StorageReader.Tell() + CapacityInBytes); // Skip over the value and its extra capacity space if any.
+				StorageReader << Key;
+				if (!IsDead(TypeCode))
+				{
+					NameOffsetMap.Emplace(Key, Offset);
+				}
+				break;
+			}
+
+			default:
+				break;
+		}
+	}
+
+	return true;
+}
+
+uint32 FAnalyticsPropertyStore::Num() const
+{
+	FScopeLock Lock(&StoreLock);
+	return NameOffsetMap.Num();
+}
+
+bool FAnalyticsPropertyStore::Contains(const FString& Key) const
+{
+	FScopeLock Lock(&StoreLock);
+	return NameOffsetMap.Contains(Key);
+}
+
+bool FAnalyticsPropertyStore::Remove(const FString& Key)
+{
+	FScopeLock Lock(&StoreLock);
+	uint32 Offset;
+	if (NameOffsetMap.RemoveAndCopyValue(Key, Offset))
+	{
+		// Update the type code, adding the flag 'dead'.
+		ETypeCode TypeCode;
+		StorageReader.Seek(Offset);
+		StorageReader.Serialize(&TypeCode, sizeof(ETypeCode));
+		TypeCode |= ETypeCode::Dead;
+		StorageWriter.Seek(Offset);
+		StorageWriter.Serialize(&TypeCode, sizeof(ETypeCode));
+		bFragmented = true;
+		return true;
+	}
+	return false; // Was not found.
+}
+
+void FAnalyticsPropertyStore::RemoveAll()
+{
+	FScopeLock Lock(&StoreLock);
+	NameOffsetMap.Reset();
+	StorageReader.Seek(0);
+	StorageWriter.Seek(0);
+	StorageBuf.Empty(StorageBuf.Num());
+	bFragmented = false;
+}
+
+IAnalyticsPropertyStore::EStatusCode FAnalyticsPropertyStore::Set(const FString& Key, int32 Value)
+{
+	FScopeLock Lock(&StoreLock);
+	return SetFixedSizeValueInternal(Key, ETypeCode::I32, Value);
+}
+
+IAnalyticsPropertyStore::EStatusCode FAnalyticsPropertyStore::Set(const FString& Key, uint32 Value)
+{
+	FScopeLock Lock(&StoreLock);
+	return SetFixedSizeValueInternal(Key, ETypeCode::U32, Value);
+}
+
+IAnalyticsPropertyStore::EStatusCode FAnalyticsPropertyStore::Set(const FString& Key, int64 Value)
+{
+	FScopeLock Lock(&StoreLock);
+	return SetFixedSizeValueInternal(Key, ETypeCode::I64, Value);
+}
+
+IAnalyticsPropertyStore::EStatusCode FAnalyticsPropertyStore::Set(const FString& Key, uint64 Value)
+{
+	FScopeLock Lock(&StoreLock);
+	return SetFixedSizeValueInternal(Key, ETypeCode::U64, Value);
+}
+
+IAnalyticsPropertyStore::EStatusCode FAnalyticsPropertyStore::Set(const FString& Key, float Value)
+{
+	FScopeLock Lock(&StoreLock);
+	return SetFixedSizeValueInternal(Key, ETypeCode::Flt, Value);
+}
+
+IAnalyticsPropertyStore::EStatusCode FAnalyticsPropertyStore::Set(const FString& Key, double Value)
+{
+	FScopeLock Lock(&StoreLock);
+	return SetFixedSizeValueInternal(Key, ETypeCode::Dbl, Value);
+}
+
+IAnalyticsPropertyStore::EStatusCode FAnalyticsPropertyStore::Set(const FString& Key, bool Value)
+{
+	FScopeLock Lock(&StoreLock);
+	return SetFixedSizeValueInternal(Key, ETypeCode::Bool, Value);
+}
+
+IAnalyticsPropertyStore::EStatusCode FAnalyticsPropertyStore::Set(const FString& Key, const FString& Value, uint32 CharCountCapacityHint)
+{
+	FScopeLock Lock(&StoreLock);
+	return SetStringValueInternal(Key, Value, CharCountCapacityHint);
+}
+
+IAnalyticsPropertyStore::EStatusCode FAnalyticsPropertyStore::Set(const FString& Key, const FDateTime& Value)
+{
+	FScopeLock Lock(&StoreLock);
+	return SetDateTimeValueInternal(Key, Value);
+}
+
+IAnalyticsPropertyStore::EStatusCode FAnalyticsPropertyStore::Set(const FString& Key, int32 Value, const TFunction<bool(const int32*, const int32&)>& ConditionFn)
+{
+	FScopeLock Lock(&StoreLock);
+
+	return AnalyticsPropertyStoreUtils::SetValue(
+				[this, &Key](int32& OutValue) { return GetFixedSizeValueInternal(Key, ETypeCode::I32, OutValue); },
+				[this, &Key](int32 InValue)   { return SetFixedSizeValueInternal(Key, ETypeCode::I32, InValue); },
+				Value,
+				ConditionFn);
+}
+
+IAnalyticsPropertyStore::EStatusCode FAnalyticsPropertyStore::Set(const FString& Key, uint32 Value, const TFunction<bool(const uint32*, const uint32&)>& ConditionFn)
+{
+	FScopeLock Lock(&StoreLock);
+
+	return AnalyticsPropertyStoreUtils::SetValue(
+				[this, &Key](uint32& OutValue) { return GetFixedSizeValueInternal(Key, ETypeCode::U32, OutValue); },
+				[this, &Key](uint32 InValue)   { return SetFixedSizeValueInternal(Key, ETypeCode::U32, InValue); },
+				Value,
+				ConditionFn);
+}
+
+IAnalyticsPropertyStore::EStatusCode FAnalyticsPropertyStore::Set(const FString& Key, int64 Value, const TFunction<bool(const int64*, const int64&)>& ConditionFn)
+{
+	FScopeLock Lock(&StoreLock);
+
+	return AnalyticsPropertyStoreUtils::SetValue(
+				[this, &Key](int64& OutValue) { return GetFixedSizeValueInternal(Key, ETypeCode::I64, OutValue); },
+				[this, &Key](int64 InValue)   { return SetFixedSizeValueInternal(Key, ETypeCode::I64, InValue); },
+				Value,
+				ConditionFn);
+}
+
+IAnalyticsPropertyStore::EStatusCode FAnalyticsPropertyStore::Set(const FString& Key, uint64 Value, const TFunction<bool(const uint64*, const uint64&)>& ConditionFn)
+{
+	FScopeLock Lock(&StoreLock);
+
+	return AnalyticsPropertyStoreUtils::SetValue(
+				[this, &Key](uint64& OutValue) { return GetFixedSizeValueInternal(Key, ETypeCode::U64, OutValue); },
+				[this, &Key](uint64 InValue)   { return SetFixedSizeValueInternal(Key, ETypeCode::U64, InValue); },
+				Value,
+				ConditionFn);
+}
+
+IAnalyticsPropertyStore::EStatusCode FAnalyticsPropertyStore::Set(const FString& Key, float Value, const TFunction<bool(const float*, const float&)>& ConditionFn)
+{
+	FScopeLock Lock(&StoreLock);
+
+	return AnalyticsPropertyStoreUtils::SetValue(
+				[this, &Key](float& OutValue) { return GetFixedSizeValueInternal(Key, ETypeCode::Flt, OutValue); },
+				[this, &Key](float InValue)   { return SetFixedSizeValueInternal(Key, ETypeCode::Flt, InValue); },
+				Value,
+				ConditionFn);
+}
+
+IAnalyticsPropertyStore::EStatusCode FAnalyticsPropertyStore::Set(const FString& Key, double Value, const TFunction<bool(const double*, const double&)>& ConditionFn)
+{
+	FScopeLock Lock(&StoreLock);
+
+	return AnalyticsPropertyStoreUtils::SetValue(
+				[this, &Key](double& OutValue) { return GetFixedSizeValueInternal(Key, ETypeCode::Dbl, OutValue); },
+				[this, &Key](double InValue)   { return SetFixedSizeValueInternal(Key, ETypeCode::Dbl, InValue); },
+				Value,
+				ConditionFn);
+}
+
+IAnalyticsPropertyStore::EStatusCode FAnalyticsPropertyStore::Set(const FString& Key, bool Value, const TFunction<bool(const bool*, const bool&)>& ConditionFn)
+{
+	FScopeLock Lock(&StoreLock);
+
+	return AnalyticsPropertyStoreUtils::SetValue(
+				[this, &Key](bool& OutValue) { return GetFixedSizeValueInternal(Key, ETypeCode::Bool, OutValue); },
+				[this, &Key](bool InValue)   { return SetFixedSizeValueInternal(Key, ETypeCode::Bool, InValue); },
+				Value,
+				ConditionFn);
+}
+
+IAnalyticsPropertyStore::EStatusCode FAnalyticsPropertyStore::Set(const FString& Key, const FString& Value, const TFunction<bool(const FString*, const FString&)>& ConditionFn)
+{
+	FScopeLock Lock(&StoreLock);
+
+	return AnalyticsPropertyStoreUtils::SetValue(
+				[this, &Key](FString& OutValue)      { return GetStringValueInternal(Key, OutValue); },
+				[this, &Key](const FString& InValue) { return SetStringValueInternal(Key, InValue); },
+				Value,
+				ConditionFn);
+}
+
+IAnalyticsPropertyStore::EStatusCode FAnalyticsPropertyStore::Set(const FString& Key, const FDateTime& Value, const TFunction<bool(const FDateTime*, const FDateTime&)>& ConditionFn)
+{
+	FScopeLock Lock(&StoreLock);
+
+	return AnalyticsPropertyStoreUtils::SetValue(
+				[this, &Key](FDateTime& OutValue)      { return GetDateTimeValueInternal(Key, OutValue); },
+				[this, &Key](const FDateTime& InValue) { return SetDateTimeValueInternal(Key, InValue); },
+				Value,
+				ConditionFn);
+}
+
+IAnalyticsPropertyStore::EStatusCode FAnalyticsPropertyStore::Update(const FString& Key, const TFunction<bool(int32&)>& UpdateFn)
+{
+	FScopeLock Lock(&StoreLock);
+
+	return AnalyticsPropertyStoreUtils::UpdateValue<int32>(
+				[this, &Key](int32& OutValue) { return GetFixedSizeValueInternal(Key, ETypeCode::I32, OutValue); },
+				[this, &Key](int32 InValue)   { return SetFixedSizeValueInternal(Key, ETypeCode::I32, InValue); },
+				UpdateFn);
+}
+
+IAnalyticsPropertyStore::EStatusCode FAnalyticsPropertyStore::Update(const FString& Key, const TFunction<bool(uint32&)>& UpdateFn)
+{
+	FScopeLock Lock(&StoreLock);
+
+	return AnalyticsPropertyStoreUtils::UpdateValue<uint32>(
+				[this, &Key](uint32& OutValue) { return GetFixedSizeValueInternal(Key, ETypeCode::U32, OutValue); },
+				[this, &Key](uint32 InValue)   { return SetFixedSizeValueInternal(Key, ETypeCode::U32, InValue); },
+				UpdateFn);
+}
+
+IAnalyticsPropertyStore::EStatusCode FAnalyticsPropertyStore::Update(const FString& Key, const TFunction<bool(int64&)>& UpdateFn)
+{
+	FScopeLock Lock(&StoreLock);
+
+	return AnalyticsPropertyStoreUtils::UpdateValue<int64>(
+				[this, &Key](int64& OutValue) { return GetFixedSizeValueInternal(Key, ETypeCode::I64, OutValue); },
+				[this, &Key](int64 InValue)   { return SetFixedSizeValueInternal(Key, ETypeCode::I64, InValue); },
+				UpdateFn);
+}
+
+IAnalyticsPropertyStore::EStatusCode FAnalyticsPropertyStore::Update(const FString& Key, const TFunction<bool(uint64&)>& UpdateFn)
+{
+	FScopeLock Lock(&StoreLock);
+
+	return AnalyticsPropertyStoreUtils::UpdateValue<uint64>(
+				[this, &Key](uint64& OutValue) { return GetFixedSizeValueInternal(Key, ETypeCode::U64, OutValue); },
+				[this, &Key](uint64 InValue)   { return SetFixedSizeValueInternal(Key, ETypeCode::U64, InValue); },
+				UpdateFn);
+}
+
+IAnalyticsPropertyStore::EStatusCode FAnalyticsPropertyStore::Update(const FString& Key, const TFunction<bool(float&)>& UpdateFn)
+{
+	FScopeLock Lock(&StoreLock);
+
+	return AnalyticsPropertyStoreUtils::UpdateValue<float>(
+				[this, &Key](float& OutValue) { return GetFixedSizeValueInternal(Key, ETypeCode::Flt, OutValue); },
+				[this, &Key](float InValue)   { return SetFixedSizeValueInternal(Key, ETypeCode::Flt, InValue); },
+				UpdateFn);
+}
+
+IAnalyticsPropertyStore::EStatusCode FAnalyticsPropertyStore::Update(const FString& Key, const TFunction<bool(double&)>& UpdateFn)
+{
+	FScopeLock Lock(&StoreLock);
+
+	return AnalyticsPropertyStoreUtils::UpdateValue<double>(
+				[this, &Key](double& OutValue) { return GetFixedSizeValueInternal(Key, ETypeCode::Dbl, OutValue); },
+				[this, &Key](double InValue)   { return SetFixedSizeValueInternal(Key, ETypeCode::Dbl, InValue); },
+				UpdateFn);
+}
+
+IAnalyticsPropertyStore::EStatusCode FAnalyticsPropertyStore::Update(const FString& Key, const TFunction<bool(bool&)>& UpdateFn)
+{
+	FScopeLock Lock(&StoreLock);
+
+	return AnalyticsPropertyStoreUtils::UpdateValue<bool>(
+				[this, &Key](bool& OutValue) { return GetFixedSizeValueInternal(Key, ETypeCode::Bool, OutValue); },
+				[this, &Key](bool InValue)   { return SetFixedSizeValueInternal(Key, ETypeCode::Bool, InValue); },
+				UpdateFn);
+}
+
+IAnalyticsPropertyStore::EStatusCode FAnalyticsPropertyStore::Update(const FString& Key, const TFunction<bool(FString&)>& UpdateFn)
+{
+	FScopeLock Lock(&StoreLock);
+
+	return AnalyticsPropertyStoreUtils::UpdateValue<FString>(
+				[this, &Key](FString& OutValue)      { return GetStringValueInternal(Key, OutValue); },
+				[this, &Key](const FString& InValue) { return SetStringValueInternal(Key, InValue); },
+				UpdateFn);
+}
+
+IAnalyticsPropertyStore::EStatusCode FAnalyticsPropertyStore::Update(const FString& Key, const TFunction<bool(FDateTime&)>& UpdateFn)
+{
+	FScopeLock Lock(&StoreLock);
+
+	return AnalyticsPropertyStoreUtils::UpdateValue<FDateTime>(
+				[this, &Key](FDateTime& OutValue)      { return GetDateTimeValueInternal(Key, OutValue); },
+				[this, &Key](const FDateTime& InValue) { return SetDateTimeValueInternal(Key, InValue); },
+				UpdateFn);
+}
+
+IAnalyticsPropertyStore::EStatusCode FAnalyticsPropertyStore::Get(const FString& Key, int32& OutValue) const
+{
+	FScopeLock Lock(&StoreLock);
+	return GetFixedSizeValueInternal(Key, ETypeCode::I32, OutValue);
+}
+
+IAnalyticsPropertyStore::EStatusCode FAnalyticsPropertyStore::Get(const FString& Key, uint32& OutValue) const
+{
+	FScopeLock Lock(&StoreLock);
+	return GetFixedSizeValueInternal(Key, ETypeCode::U32, OutValue);
+}
+
+IAnalyticsPropertyStore::EStatusCode FAnalyticsPropertyStore::Get(const FString& Key, int64& OutValue) const
+{
+	FScopeLock Lock(&StoreLock);
+	return GetFixedSizeValueInternal(Key, ETypeCode::I64, OutValue);
+}
+
+IAnalyticsPropertyStore::EStatusCode FAnalyticsPropertyStore::Get(const FString& Key, uint64& OutValue) const
+{
+	FScopeLock Lock(&StoreLock);
+	return GetFixedSizeValueInternal(Key, ETypeCode::U64, OutValue);
+}
+
+IAnalyticsPropertyStore::EStatusCode FAnalyticsPropertyStore::Get(const FString& Key, float& OutValue) const
+{
+	FScopeLock Lock(&StoreLock);
+	return GetFixedSizeValueInternal(Key, ETypeCode::Flt, OutValue);
+}
+
+IAnalyticsPropertyStore::EStatusCode FAnalyticsPropertyStore::Get(const FString& Key, double& OutValue) const
+{
+	FScopeLock Lock(&StoreLock);
+	return GetFixedSizeValueInternal(Key, ETypeCode::Dbl, OutValue);
+}
+
+IAnalyticsPropertyStore::EStatusCode FAnalyticsPropertyStore::Get(const FString& Key, bool& OutValue) const
+{
+	FScopeLock Lock(&StoreLock);
+	return GetFixedSizeValueInternal(Key, ETypeCode::Bool, OutValue);
+}
+
+IAnalyticsPropertyStore::EStatusCode FAnalyticsPropertyStore::Get(const FString& Key, FString& OutValue) const
+{
+	FScopeLock Lock(&StoreLock);
+	return GetStringValueInternal(Key, OutValue);
+}
+
+IAnalyticsPropertyStore::EStatusCode FAnalyticsPropertyStore::Get(const FString& Key, FDateTime& OutValue) const
+{
+	FScopeLock Lock(&StoreLock);
+	return GetDateTimeValueInternal(Key, OutValue);
+}
+
+IAnalyticsPropertyStore::EStatusCode FAnalyticsPropertyStore::SetFixedSizeValueInternal(const FString& Key, ETypeCode TypeCode, const uint8* Value, uint32 Size)
+{
+	check(IsFixedSize(TypeCode));
+
+	// NOTE: Format for records holding fixed-length value is: [TypeCode | Value | Key ]
+	if (const uint32* RecordOffset = NameOffsetMap.Find(Key))
+	{
+		ETypeCode StoredValueType;
+		StorageReader.Seek(*RecordOffset);
+		StorageReader.Serialize(&StoredValueType, sizeof(ETypeCode));
+		if (StoredValueType == TypeCode)
+		{
+			StorageWriter.Seek(*RecordOffset + sizeof(TypeCode));
+			StorageWriter.Serialize(const_cast<uint8*>(Value), Size);
+			return EStatusCode::Success;
+		}
+		return EStatusCode::BadType;
+	}
+	else
+	{
+		// Insert the new record at the end.
+		uint32 NewRecordOffset = StorageWriter.TotalSize();
+		StorageWriter.Seek(NewRecordOffset);
+
+		// Write the type.
+		StorageWriter.Serialize(&TypeCode, sizeof(ETypeCode));
+
+		// Write the value.
+		StorageWriter.Serialize(const_cast<uint8*>(Value), Size);
+
+		// Write the key. Let the FArchive encode the string.
+		StorageWriter << const_cast<FString&>(Key);
+
+		// Cache the offset to the newly inserted record.
+		NameOffsetMap.Emplace(Key, NewRecordOffset);
+		return EStatusCode::Success;
+	}
+}
+
+IAnalyticsPropertyStore::EStatusCode FAnalyticsPropertyStore::GetFixedSizeValueInternal(const FString& Key, ETypeCode TypeCode, uint8* OutValue, uint32 Size) const
+{
+	check(IsFixedSize(TypeCode));
+
+	// NOTE: Format for records holding fixed-length value is: [TypeCode | Value | Key ]
+	if (const uint32* RecordOffset = NameOffsetMap.Find(Key))
+	{
+		ETypeCode StoredTypeCode;
+		StorageReader.Seek(*RecordOffset);
+		StorageReader.Serialize(&StoredTypeCode, sizeof(ETypeCode));
+		if (StoredTypeCode == TypeCode)
+		{
+			StorageReader.Serialize(OutValue, Size);
+			return EStatusCode::Success;
+		}
+		return EStatusCode::BadType;
+	}
+	return EStatusCode::NotFound;
+}
+
+IAnalyticsPropertyStore::EStatusCode FAnalyticsPropertyStore::SetStringValueInternal(const FString& Key, const FString& Value, uint32 CharCountCapacityHint)
+{
+	// NOTE: Format for records holding strings is: [Type | SizeInBytes | CapacityInBytes | String-UTF16 | Key ]
+
+	// Insert a variable-length value.
+	auto InsertFn = [this](const FString& Key, ETypeCode TypeCode, const void* Bytes, uint32 ByteCount, uint32 Capacity)
+	{
+		// Insert the new record at the end.
+		uint32 InsertOffset = StorageWriter.TotalSize();
+		StorageWriter.Seek(InsertOffset);
+
+		// Write the value type.
+		StorageWriter.Serialize(&TypeCode, sizeof(ETypeCode));
+
+		// Write the data size.
+		StorageWriter.Serialize(const_cast<uint32*>(&ByteCount), sizeof(uint32));
+
+		// Write the capacity (in bytes) and reserve space if needed.
+		StorageWriter.Serialize(&Capacity, sizeof(uint32));
+
+		// Write the data.
+		StorageWriter.Serialize(const_cast<void*>(Bytes), ByteCount);
+
+		// Reserve extra capacity if needed.
+		if (Capacity > ByteCount)
+		{
+			StorageWriter.Seek(StorageWriter.TotalSize() + Capacity - ByteCount); // Reserve extra space for the value.
+		}
+
+		// Write the key. Let the FArchive encode the string.
+		StorageWriter << const_cast<FString&>(Key);
+
+		// Cache the offset to the newly inserted record.
+		NameOffsetMap.Emplace(Key, InsertOffset);
+		return EStatusCode::Success;
+	};
+
+	// Store the string in UTF16 with a null terminator. This is a no-op on platforms that are using a 16-bit TCHAR.
+	FTCHARToUTF16 UTF16String(*Value, Value.Len() + 1);
+	const void* Data = UTF16String.Get();
+	uint32 DataSize = (UTF16String.Length() + 1) * sizeof(UTF16CHAR);
+
+	if (const uint32* RecordOffset = NameOffsetMap.Find(Key))
+	{
+		// Ensure the types match.
+		ETypeCode StoredValueType;
+		StorageReader.Seek(*RecordOffset);
+		StorageReader.Serialize(&StoredValueType, sizeof(ETypeCode));
+		if (StoredValueType == ETypeCode::Str)
+		{
+			// Skip over the stored size.
+			StorageReader.Seek(StorageReader.Tell() + sizeof(uint32));
+
+			// Read the capacity.
+			uint32 AvailableCapacity;
+			StorageReader.Serialize(&AvailableCapacity, sizeof(uint32));
+
+			// Can it be updated in-place?
+			if (AvailableCapacity >= DataSize)
+			{
+				// Update the size.
+				StorageWriter.Seek(*RecordOffset + sizeof(ETypeCode));
+				StorageWriter.Serialize(&DataSize, sizeof(uint32));
+
+				// Skip over the capacity (it did not change).
+				StorageWriter.Seek(StorageWriter.Tell() + sizeof(uint32));
+
+				// Update the value.
+				StorageWriter.Serialize(const_cast<void*>(Data), DataSize);
+				return EStatusCode::Success;
+			}
+			else
+			{
+				// Tombstone the current location.
+				ETypeCode Tombstone = ETypeCode::Str | ETypeCode::Dead;
+				StorageWriter.Seek(*RecordOffset);
+				StorageWriter.Serialize(&Tombstone, sizeof(ETypeCode));
+				bFragmented = true;
+
+				// Relocate the record at the end.
+				return InsertFn(Key, ETypeCode::Str, Data, DataSize, DataSize);
+			}
+		}
+		return EStatusCode::BadType;
+	}
+
+	constexpr uint32 MaxExtraCapacityReserved = 2 * 1024;
+	const uint32 RequestedCapacity = CharCountCapacityHint * static_cast<uint32>(sizeof(UTF16CHAR));
+	return InsertFn(Key, ETypeCode::Str, Data, DataSize, RequestedCapacity <= DataSize ? DataSize : FMath::Min(RequestedCapacity, DataSize + MaxExtraCapacityReserved));
+}
+
+IAnalyticsPropertyStore::EStatusCode FAnalyticsPropertyStore::GetStringValueInternal(const FString& Key, FString& OutValue) const
+{
+	// NOTE: Format for records holding strings is: [Type | SizeInBytes | CapacityInBytes | UTF16String | Key ]
+
+	if (const uint32* RecordOffset = NameOffsetMap.Find(Key))
+	{
+		ETypeCode StoredTypeCode;
+		StorageReader.Seek(*RecordOffset);
+		StorageReader.Serialize(&StoredTypeCode, sizeof(ETypeCode));
+		if (StoredTypeCode == ETypeCode::Str)
+		{
+			// Read the size.
+			uint32 DataSize;
+			StorageReader.Serialize(&DataSize, sizeof(uint32));
+
+			// Skip over the capacity.
+			StorageReader.Seek(StorageReader.Tell() + sizeof(uint32));
+
+			// Read the UTF16 string into a FString.
+			uint32 CharCountWithNull = DataSize / sizeof(UTF16CHAR);
+			OutValue.GetCharArray().Empty(CharCountWithNull);
+			OutValue.GetCharArray().AddUninitialized(CharCountWithNull);
+			auto Passthru = StringMemoryPassthru<UTF16CHAR>(OutValue.GetCharArray().GetData(), CharCountWithNull, CharCountWithNull);
+			StorageReader.Serialize(Passthru.Get(), DataSize);
+			Passthru.Get()[CharCountWithNull - 1] = '\0';
+			Passthru.Apply();
+
+			// Inline combine any surrogate pairs in the data when loading into a UTF-32 string
+			StringConv::InlineCombineSurrogates(OutValue);
+
+			return EStatusCode::Success;
+		}
+		return EStatusCode::BadType;
+	}
+	return EStatusCode::NotFound;
+}
+
+IAnalyticsPropertyStore::EStatusCode FAnalyticsPropertyStore::SetDateTimeValueInternal(const FString& Key, const FDateTime& Value)
+{
+	int64 Ticks = Value.GetTicks();
+	return SetFixedSizeValueInternal(Key, ETypeCode::Date, reinterpret_cast<const uint8*>(&Ticks), sizeof(Ticks));
+}
+
+IAnalyticsPropertyStore::EStatusCode FAnalyticsPropertyStore::GetDateTimeValueInternal(const FString& Key, FDateTime& OutValue) const
+{
+	int64 Ticks;
+	EStatusCode Status = GetFixedSizeValueInternal(Key, ETypeCode::Date, Ticks);
+	if (Status == EStatusCode::Success)
+	{
+		OutValue = FDateTime(Ticks);
+	}
+	return Status;
+}
+
+void FAnalyticsPropertyStore::Defragment()
+{
+	// NOTE: This is not implemented because we don't fragment enough the store in the current usage to benefit from defragmenting it.
+	bFragmented = false;
+}
+
+void FAnalyticsPropertyStore::VisitAll(const TFunction<void(const FString&, FString&&)>& VisitFn) const
+{
+	FScopeLock Lock(&StoreLock);
+
+	for (const TPair<FString, uint32>& NameOffsetPair : NameOffsetMap)
+	{
+		ETypeCode TypeCode;
+		StorageReader.Seek(NameOffsetPair.Value);
+		StorageReader << TypeCode;
+		switch (RawType(TypeCode))
+		{
+			case ETypeCode::I32:
+			{
+				int32 V;
+				StorageReader.Serialize(&V, sizeof(V));
+				VisitFn(NameOffsetPair.Key, LexToString(V));
+				break;
+			}
+
+			case ETypeCode::U32:
+			{
+				uint32 V;
+				StorageReader.Serialize(&V, sizeof(V));
+				VisitFn(NameOffsetPair.Key, LexToString(V));
+				break;
+			}
+
+			case ETypeCode::I64:
+			{
+				int64 V;
+				StorageReader.Serialize(&V, sizeof(V));
+				VisitFn(NameOffsetPair.Key, LexToString(V));
+				break;
+			}
+
+			case ETypeCode::U64:
+			{
+				uint64 V;
+				StorageReader.Serialize(&V, sizeof(V));
+				VisitFn(NameOffsetPair.Key, LexToString(V));
+				break;
+			}
+
+			case ETypeCode::Flt:
+			{
+				float V;
+				StorageReader.Serialize(&V, sizeof(V));
+				VisitFn(NameOffsetPair.Key, FString::SanitizeFloat(V));
+				break;
+			}
+
+			case ETypeCode::Dbl:
+			{
+				double V;
+				StorageReader.Serialize(&V, sizeof(V));
+				VisitFn(NameOffsetPair.Key, FString::SanitizeFloat(V));
+				break;
+			}
+
+			case ETypeCode::Bool:
+			{
+				bool V;
+				StorageReader.Serialize(&V, sizeof(V));
+				VisitFn(NameOffsetPair.Key, LexToString(V)); // Convert to "true" or "false".
+				break;
+			}
+
+			case ETypeCode::Str:
+			{
+				FString Value;
+				GetStringValueInternal(NameOffsetPair.Key, Value);
+				VisitFn(NameOffsetPair.Key, MoveTemp(Value));
+				break;
+			}
+
+			case ETypeCode::Date:
+			{
+				FDateTime Value;
+				GetDateTimeValueInternal(NameOffsetPair.Key, Value);
+				VisitFn(NameOffsetPair.Key, Value.ToIso8601()); // Iso-8601 is the format used by Analytics in general for date/time.
+				break;
+			}
+
+			default:
+				break;
+		}
+	}
+}
+
+bool FAnalyticsPropertyStore::Flush(bool bAsync, const FTimespan& Timeout)
+{
+	if (!IsValid())
+	{
+		return false;
+	}
+
+	TRACE_CPUPROFILER_EVENT_SCOPE(FAnalyticsPropertyStore::Flush);
+
+	// Local helper function to flush.
+	auto FlushLockedFn = [this, bAsync](bool bNoWaiting)
+	{
+		// Ensure the previous task completed. Only one thread at a time can write the file.
+		if (FlushTask.IsValid())
+		{
+			// If the task completed or can be successufully canceled.
+			if (FlushTask->IsDone() || FlushTask->Cancel())
+			{
+				FlushTask.Reset();
+			}
+			else if (bNoWaiting)
+			{
+				return false; // Not allowed to wait. Don't flush.
+			}
+			else
+			{
+				// Wait for task to complete.
+				FlushTask->EnsureCompletion();
+				FlushTask.Reset();
+			}
+		}
+
+		if (bAsync)
+		{
+			// Copy the storage buffer in a reusable buffer since there is at most one async flush at any point in time. This prevent allocating,
+			// buffer sizes match most of the time.
+			AsyncFlushDataCopy = StorageBuf;
+			FlushTask = MakeUnique<FAsyncTask<FFlushWorker>>(*this);
+			if (FlushTask)
+			{
+				FlushTask->StartBackgroundTask();
+				return true; // Task started.
+			}
+		}
+
+		// Flush synchronously.
+		FlushInternal(StorageBuf);
+		return true; // Flushed.
+	};
+
+	// If the caller doesn't want to wait to get its request started.
+	if (Timeout.IsZero() && StoreLock.TryLock())
+	{
+		ON_SCOPE_EXIT { StoreLock.Unlock(); };
+		return FlushLockedFn(/*bNoWaiting*/true);
+	}
+	// If the caller is ready to wait indefinitely to get its request started.
+	else if (Timeout == FTimespan::MaxValue())
+	{
+		FScopeLock ScopeLock(&StoreLock);
+		return FlushLockedFn(/*bNoWaiting*/false);
+	}
+	// The caller is willing to wait some time to get its request processed.
+	else
+	{
+		double StartTimeSecs = FPlatformTime::Seconds();
+		do
+		{
+			if (StoreLock.TryLock())
+			{
+				ON_SCOPE_EXIT { StoreLock.Unlock(); };
+				if (FlushLockedFn(/*bNoWaiting*/true))
+				{
+					return true;
+				}
+			}
+			FPlatformProcess::Sleep(FTimespan::FromMilliseconds(1).GetTotalSeconds());
+		} while (FPlatformTime::Seconds() - StartTimeSecs < Timeout.GetTotalSeconds());
+
+		return false;
+	}
+}
+
+void FAnalyticsPropertyStore::FlushInternal(const TArray<uint8>& Data)
+{
+	AnalyticsPropertyStoreUtils::FAnalyticsStoreFileHeader Header;
+	Header.Version = AnalyticsPropertyStoreUtils::StoreFormatVersion;
+	Header.DataSize = Data.Num();
+
+	// Compute the checksum.
+	FMD5 Md5;
+	Md5.Update(Data.GetData(), Data.Num());
+	Md5.Final(Header.Checksum);
+
+	// Write the file header.
+	FileHandle->Seek(0);
+	FileHandle->Write(reinterpret_cast<const uint8*>(&Header), sizeof(Header));
+
+	// Write storage buf.
+	FileHandle->Write(Data.GetData(), Data.Num());
+
+	// Flush the File.
+	FileHandle->Flush();
+}
+
+FAnalyticsPropertyStore::FFlushWorker::FFlushWorker(FAnalyticsPropertyStore& InStore)
+	: Store(InStore)
+{
+}
+
+void FAnalyticsPropertyStore::FFlushWorker::DoWork()
+{
+	// This expects that only one thread/task at the time flushes the data.
+	Store.FlushInternal(Store.AsyncFlushDataCopy);
+}
diff --git a/Engine/Source/Runtime/Analytics/Analytics/Private/Tests/AnalyticsPropertyStoreTests.cpp b/Engine/Source/Runtime/Analytics/Analytics/Private/Tests/AnalyticsPropertyStoreTests.cpp
new file mode 100644
--- /dev/null
+++ b/Engine/Source/Runtime/Analytics/Analytics/Private/Tests/AnalyticsPropertyStoreTests.cpp
@@ -0,0 +1,397 @@
+// Copyright Epic Games, Inc. All Rights Reserved.
+
+#include "CoreMinimal.h"
+#include "GenericPlatform/GenericPlatformFile.h"
+#include "HAL/PlatformFileManager.h"
+#include "Misc/AutomationTest.h"
+#include "AnalyticsPropertyStore.h"
+#include "CoreGlobals.h"
+
+#if WITH_DEV_AUTOMATION_TESTS
+
+IMPLEMENT_SIMPLE_AUTOMATION_TEST(FAnalyticsPropertyStoreAutomationTest, "System.Analytics.PropertyStore", EAutomationTestFlags::EditorContext | EAutomationTestFlags::EngineFilter | EAutomationTestFlags::MediumPriority)
+
+bool FAnalyticsPropertyStoreAutomationTest::RunTest(const FString& Parameters)
+{
+	static const FString I32Key  = TEXT("MyInt32");
+	static const FString U32Key  = TEXT("MyUInt32");
+	static const FString I64Key  = TEXT("MyIint64");
+	static const FString U64Key  = TEXT("MyUint64");
+	static const FString FltKey  = TEXT("MyFloat");
+	static const FString DblKey  = TEXT("MyDouble");
+	static const FString BoolKey = TEXT("MyBool");
+	static const FString Str1Key = TEXT("MyStr1");
+	static const FString Str2Key = TEXT("MyStr2");
+	static const FString DateKey = TEXT("MyDate");
+
+	int32 Expected_I32 = 10;
+	uint32 Expected_U32 = 20;
+	int64 Expected_I64 = 30;
+	uint64 Expected_U64 = 40;
+	float Expected_Flt = 50.0f;
+	double Expected_Dbl = 60.0;
+	bool Expected_Bool = true;
+	FString Expected_Str1 = "Hello World";
+	FString Expected_Str2 = "Just Do It";
+	FDateTime Expected_Date = FDateTime::UtcNow();
+
+	int32 Actual_I32;
+	uint32 Actual_U32;
+	int64 Actual_I64;
+	uint64 Acutal_U64;
+	float Actual_Flt;
+	double Actual_Dbl;
+	bool Actual_Bool;
+	FString Actual_Str1;
+	FString Actual_Str2;
+	FDateTime Actual_Date;
+	int32 NotFoundDummy;
+
+	// Ensure the automation directory exists.
+	IFileManager::Get().MakeDirectory(*FPaths::AutomationTransientDir(), /*Tree*/true);
+	FString TestStorePathname = FPaths::AutomationTransientDir() / TEXT("AnalyticsPropertyStoreTest.bin");
+
+	// Test invalid store - not created/not loaded.
+	{
+		FAnalyticsPropertyStore Store;
+		check(!Store.IsValid());
+		check(Store.Num() == 0);
+	}
+
+	// Create an empty store and reload it.
+	{
+		FAnalyticsPropertyStore Store;
+		check(Store.Create(TestStorePathname, /*CapacityHint*/2 * 1024));
+		check(Store.IsValid());
+		check(Store.Num() == 0);
+		Store.Flush();
+		check(Store.Load(TestStorePathname));
+		check(Store.IsValid());
+		check(Store.Num() == 0);
+	}
+
+	// Test the basic operations.
+	{
+		FAnalyticsPropertyStore Store;
+		check(Store.Create(TestStorePathname, /*CapacityHint*/2 * 1024));
+		check(Store.IsValid());
+		check(Store.Num() == 0);
+
+		// Write values.
+		check(Store.Set(I32Key,  Expected_I32) == IAnalyticsPropertyStore::EStatusCode::Success);
+		check(Store.Set(U32Key,  Expected_U32) == IAnalyticsPropertyStore::EStatusCode::Success);
+		check(Store.Set(I64Key,  Expected_I64) == IAnalyticsPropertyStore::EStatusCode::Success);
+		check(Store.Set(U64Key,  Expected_U64) == IAnalyticsPropertyStore::EStatusCode::Success);
+		check(Store.Set(FltKey,  Expected_Flt) == IAnalyticsPropertyStore::EStatusCode::Success);
+		check(Store.Set(DblKey,  Expected_Dbl) == IAnalyticsPropertyStore::EStatusCode::Success);
+		check(Store.Set(BoolKey, Expected_Bool) == IAnalyticsPropertyStore::EStatusCode::Success);
+		check(Store.Set(Str1Key, Expected_Str1, /*CapacityInChars*/25) == IAnalyticsPropertyStore::EStatusCode::Success);
+		check(Store.Set(DateKey, Expected_Date) == IAnalyticsPropertyStore::EStatusCode::Success);
+		check(Store.Set(Str2Key, Expected_Str2) == IAnalyticsPropertyStore::EStatusCode::Success);
+
+		// Read values
+		check(Store.Get(I32Key,  Actual_I32)  == IAnalyticsPropertyStore::EStatusCode::Success && Expected_I32 == Actual_I32);
+		check(Store.Get(U32Key,  Actual_U32)  == IAnalyticsPropertyStore::EStatusCode::Success && Expected_U32 == Actual_U32);
+		check(Store.Get(I64Key,  Actual_I64)  == IAnalyticsPropertyStore::EStatusCode::Success && Expected_I64 == Actual_I64);
+		check(Store.Get(U64Key,  Acutal_U64)  == IAnalyticsPropertyStore::EStatusCode::Success && Expected_U64 == Acutal_U64);
+		check(Store.Get(FltKey,  Actual_Flt)  == IAnalyticsPropertyStore::EStatusCode::Success && Expected_Flt == Actual_Flt);
+		check(Store.Get(DblKey,  Actual_Dbl)  == IAnalyticsPropertyStore::EStatusCode::Success && Expected_Dbl == Actual_Dbl);
+		check(Store.Get(BoolKey, Actual_Bool) == IAnalyticsPropertyStore::EStatusCode::Success && Expected_Bool == Actual_Bool);
+		check(Store.Get(Str1Key, Actual_Str1) == IAnalyticsPropertyStore::EStatusCode::Success && Expected_Str1 == Actual_Str1);
+		check(Store.Get(DateKey, Actual_Date) == IAnalyticsPropertyStore::EStatusCode::Success && Expected_Date == Actual_Date);
+		check(Store.Get(Str2Key, Actual_Str2) == IAnalyticsPropertyStore::EStatusCode::Success && Expected_Str2 == Actual_Str2);
+
+		// Key not found.
+		check(Store.Get(TEXT("ShouldNotFound"), NotFoundDummy) == IAnalyticsPropertyStore::EStatusCode::NotFound);
+
+		// Type mismatch.
+		check(Store.Get(I32Key, Actual_U32)  == IAnalyticsPropertyStore::EStatusCode::BadType && Actual_U32 == Expected_U32);
+		check(Store.Get(Str2Key, Actual_Dbl) == IAnalyticsPropertyStore::EStatusCode::BadType && Actual_Dbl == Expected_Dbl);
+
+		// Contains.
+		check(Store.Contains(I32Key));
+		check(Store.Contains(U32Key));
+		check(Store.Contains(I64Key));
+		check(Store.Contains(U64Key));
+		check(Store.Contains(FltKey));
+		check(Store.Contains(DblKey));
+		check(Store.Contains(BoolKey));
+		check(Store.Contains(Str1Key));
+		check(Store.Contains(DateKey));
+		check(Store.Contains(Str2Key));
+		check(!Store.Contains(TEXT("IDoNotExist")));
+
+		// In-place update.
+		Expected_I32 = 2000;
+		check(Store.Set(I32Key, Expected_I32) == IAnalyticsPropertyStore::EStatusCode::Success);
+		check(Store.Get(I32Key, Actual_I32) == IAnalyticsPropertyStore::EStatusCode::Success);
+		check(Expected_I32 == Actual_I32);
+
+		// In-place string update. (Enough capacity was reserved)
+		Expected_Str1 = TEXT("Hello World Update");
+		check(Store.Set(Str1Key, Expected_Str1) == IAnalyticsPropertyStore::EStatusCode::Success);
+		check(Store.Get(Str1Key, Actual_Str1) == IAnalyticsPropertyStore::EStatusCode::Success && Actual_Str1 == Expected_Str1);
+		check(Store.Get(Str2Key, Actual_Str2) == IAnalyticsPropertyStore::EStatusCode::Success && Expected_Str2 == Actual_Str2); // Ensure Str1 did not overwrite Str2.
+
+		// Out-of-place string update. (Not enough capacity to perform in-place update).
+		Expected_Str2 = TEXT("Hello World Extend");
+		check(Store.Set(Str2Key, Expected_Str2) == IAnalyticsPropertyStore::EStatusCode::Success);
+		check(Store.Get(Str2Key, Actual_Str2) == IAnalyticsPropertyStore::EStatusCode::Success && Expected_Str2 == Actual_Str2);
+
+		// Set conditionnaly - accepted
+		++Expected_I32;
+		check(Store.Set(I32Key, Expected_I32, [](const int32* Actual, const int32& Proposed) { return *Actual < Proposed; }) == IAnalyticsPropertyStore::EStatusCode::Success);
+		check(Store.Get(I32Key, Actual_I32) == IAnalyticsPropertyStore::EStatusCode::Success && Expected_I32 == Actual_I32);
+
+		// Set conditionnaly - declined
+		check(Store.Set(FltKey, Expected_Flt - 50.0f, [](const float* Actual, const float& Proposed) { return *Actual < Proposed; }) == IAnalyticsPropertyStore::EStatusCode::Declined);
+		check(Store.Get(FltKey, Actual_Flt) == IAnalyticsPropertyStore::EStatusCode::Success && Actual_Flt == Expected_Flt);
+
+		// Set conditionnaly - wrong type
+		check(Store.Set(BoolKey, Expected_Dbl, [](const double* Actual, const double& Proposed) { check(false); return *Actual < Proposed; }) == IAnalyticsPropertyStore::EStatusCode::BadType);
+
+		// Set conditionnaly - check params
+		Expected_Str1 = TEXT("Hi, this is Joe, how are you?");
+		check(Store.Set(Str1Key, Expected_Str1, [&](const FString* Actual, const FString& Proposed) { check(*Actual == Actual_Str1); check(Proposed == Expected_Str1); return true; }) == IAnalyticsPropertyStore::EStatusCode::Success);
+		check(Store.Get(Str1Key, Actual_Str1) == IAnalyticsPropertyStore::EStatusCode::Success && Actual_Str1 == Expected_Str1);
+
+		// Update - accepted.
+		++Expected_I64;
+		check(Store.Update(I64Key, [](int64& InOutStored) { ++InOutStored; return true; }) == IAnalyticsPropertyStore::EStatusCode::Success);
+		check(Store.Get(I64Key, Actual_I64) == IAnalyticsPropertyStore::EStatusCode::Success && Actual_I64 == Expected_I64);
+
+		// Update - declined.
+		check(Store.Update(DateKey, [](FDateTime& InOutStored) { InOutStored = FDateTime::UtcNow(); return false; }) == IAnalyticsPropertyStore::EStatusCode::Declined);
+		check(Store.Get(DateKey, Actual_Date) == IAnalyticsPropertyStore::EStatusCode::Success && Actual_Date == Expected_Date);
+
+		// Update - wrong type
+		check(Store.Update(Str1Key, [](uint32& InOutStored) { check(false); return true; }) == IAnalyticsPropertyStore::EStatusCode::BadType);
+
+		// Update - check params
+		Expected_Bool = !Expected_Bool;
+		check(Store.Update(BoolKey, [&](bool& InOutStored) { check(InOutStored == Actual_Bool); InOutStored = Expected_Bool; return true; }) == IAnalyticsPropertyStore::EStatusCode::Success);
+		check(Store.Get(BoolKey, Actual_Bool) == IAnalyticsPropertyStore::EStatusCode::Success && Actual_Bool == Expected_Bool);
+
+		// Test string conversion.
+		Store.VisitAll([&](const FString& Name, FString&& Value)
+		{
+			if (Name == I32Key)
+			{
+				check(Value == LexToString(Expected_I32));
+			}
+			else if (Name == U32Key)
+			{
+				check(Value == LexToString(Expected_U32));
+			}
+			else if (Name == I64Key)
+			{
+				check(Value == LexToString(Expected_I64));
+			}
+			else if (Name == U64Key)
+			{
+				check(Value == LexToString(Expected_U64));
+			}
+			else if (Name == FltKey)
+			{
+				check(Value == FString::SanitizeFloat(Expected_Flt));
+			}
+			else if (Name == DblKey)
+			{
+				check(Value == FString::SanitizeFloat(Expected_Dbl));
+			}
+			else if (Name == BoolKey)
+			{
+				check(Value == LexToString(Expected_Bool));
+			}
+			else if (Name == Str1Key)
+			{
+				check(Value == Expected_Str1);
+			}
+			else if (Name == DateKey)
+			{
+				check(Value == Expected_Date.ToIso8601());
+			}
+			else if (Name == Str2Key)
+			{
+				check(Value == Expected_Str2);
+			}
+			else
+			{
+				check(false); // This is an unexpected key...
+			}
+		});
+
+		// Persist the changes above to disk.
+		Store.Flush();
+	}
+
+	// Reload the store written above and verify the values.
+	{
+		FAnalyticsPropertyStore Store;
+		check(Store.Load(TestStorePathname));
+		check(Store.IsValid());
+		check(Store.Num() == 10);
+
+		check(Store.Get(I32Key,  Actual_I32)  == IAnalyticsPropertyStore::EStatusCode::Success && Expected_I32 == Actual_I32);
+		check(Store.Get(U32Key,  Actual_U32)  == IAnalyticsPropertyStore::EStatusCode::Success && Expected_U32 == Actual_U32);
+		check(Store.Get(I64Key,  Actual_I64)  == IAnalyticsPropertyStore::EStatusCode::Success && Expected_I64 == Actual_I64);
+		check(Store.Get(U64Key,  Acutal_U64)  == IAnalyticsPropertyStore::EStatusCode::Success && Expected_U64 == Acutal_U64);
+		check(Store.Get(FltKey,  Actual_Flt)  == IAnalyticsPropertyStore::EStatusCode::Success && Expected_Flt == Actual_Flt);
+		check(Store.Get(DblKey,  Actual_Dbl)  == IAnalyticsPropertyStore::EStatusCode::Success && Expected_Dbl == Actual_Dbl);
+		check(Store.Get(BoolKey, Actual_Bool) == IAnalyticsPropertyStore::EStatusCode::Success && Expected_Bool == Actual_Bool);
+		check(Store.Get(Str1Key, Actual_Str1) == IAnalyticsPropertyStore::EStatusCode::Success && Expected_Str1 == Actual_Str1);
+		check(Store.Get(DateKey, Actual_Date) == IAnalyticsPropertyStore::EStatusCode::Success && Expected_Date == Actual_Date);
+		check(Store.Get(Str2Key, Actual_Str2) == IAnalyticsPropertyStore::EStatusCode::Success && Expected_Str2 == Actual_Str2);
+
+		// Test the remove API
+		check(Store.Remove(U32Key));
+		check(Store.Remove(I32Key));
+		check(Store.Remove(Str1Key));
+		check(Store.Remove(I64Key));
+		check(!Store.Remove(I64Key)); // Remove same key twice.
+		check(Store.Remove(U64Key));
+		check(Store.Remove(DblKey));
+		check(Store.Remove(FltKey));
+		check(Store.Remove(BoolKey));
+		check(Store.Remove(Str2Key));
+		check(!Store.Remove(Str2Key)); // Remove same key twice.
+		check(Store.Remove(DateKey));
+		check(!Store.Remove(TEXT("IDoNotExist")));
+
+		check(Store.Num() == 0);
+
+		// Ensure the key were deleted.
+		check(!Store.Contains(I32Key));
+		check(!Store.Contains(U32Key));
+		check(!Store.Contains(I64Key));
+		check(!Store.Contains(U64Key));
+		check(!Store.Contains(FltKey));
+		check(!Store.Contains(DblKey));
+		check(!Store.Contains(BoolKey));
+		check(!Store.Contains(Str1Key));
+		check(!Store.Contains(DateKey));
+		check(!Store.Contains(Str2Key));
+
+		// Persist the empty store to disk.
+		Store.Flush();
+	}
+
+	// Reload an empty store and check it doesn't have any key.
+	{
+		FAnalyticsPropertyStore Store;
+		check(Store.Load(TestStorePathname));
+		check(Store.Num() == 0);
+
+		check(!Store.Contains(I32Key));
+		check(!Store.Contains(U32Key));
+		check(!Store.Contains(I64Key));
+		check(!Store.Contains(U64Key));
+		check(!Store.Contains(FltKey));
+		check(!Store.Contains(DblKey));
+		check(!Store.Contains(BoolKey));
+		check(!Store.Contains(Str1Key));
+		check(!Store.Contains(DateKey));
+		check(!Store.Contains(Str2Key));
+
+		// Add if null.
+		check(Store.Set(I64Key, Expected_I64, [](const int64* Actual, const int64& Proposed) { return Actual == nullptr; }) == IAnalyticsPropertyStore::EStatusCode::Success);
+		check(Store.Get(I64Key, Actual_I64) == IAnalyticsPropertyStore::EStatusCode::Success && Actual_I64 == Expected_I64);
+
+		// Decline if null
+		check(Store.Set(I32Key, Expected_I32, [](const int32* Actual, const int32& Proposed) { return Actual != nullptr; }) == IAnalyticsPropertyStore::EStatusCode::Declined);
+		check(Store.Get(I32Key, Actual_I32) == IAnalyticsPropertyStore::EStatusCode::NotFound && Actual_I32 == Expected_I32);
+
+		check(Store.Num() == 1);
+
+		check(Store.Set(Str2Key, Expected_Str2) == IAnalyticsPropertyStore::EStatusCode::Success);
+		check(Store.Set(DblKey,  Expected_Dbl) == IAnalyticsPropertyStore::EStatusCode::Success);
+		check(Store.Set(BoolKey, Expected_Bool) == IAnalyticsPropertyStore::EStatusCode::Success);
+
+		// Remove all.
+		Store.RemoveAll();
+		check(Store.Num() == 0);
+		check(!Store.Contains(I64Key));
+		check(!Store.Contains(DblKey));
+		check(!Store.Contains(Str2Key));
+		check(!Store.Contains(BoolKey))
+	}
+
+	// Corrupt a store.
+	{
+		TUniquePtr<FAnalyticsPropertyStore> Store = MakeUnique<FAnalyticsPropertyStore>();
+		check(Store->Create(TestStorePathname, /*CapacityHint*/2 * 1024));
+
+		Store->Set(I32Key,  Expected_I32);
+		Store->Set(U32Key,  Expected_U32);
+		Store->Set(I64Key,  Expected_I64);
+		Store->Set(U64Key,  Expected_U64);
+		Store->Set(FltKey,  Expected_Flt);
+		Store->Set(DblKey,  Expected_Dbl);
+		Store->Set(BoolKey, Expected_Bool);
+		Store->Set(Str1Key, Expected_Str1, /*CapacityInChars*/25);
+		Store->Set(DateKey, Expected_Date);
+		Store->Set(Str2Key, Expected_Str2);
+		Store->Flush();
+		Store.Reset(); // Close the file.
+
+		TUniquePtr<IFileHandle> FileHandle(FPlatformFileManager::Get().GetPlatformFile().OpenWrite(*TestStorePathname, /*bAppend*/true, /*bAllowRead*/true));
+		FileHandle->Seek(55);
+		uint32 Corruption[] = { 66, 99, 1, 7, 44, 8888 }; // Random stream of bytes.
+		FileHandle->Write(reinterpret_cast<uint8*>(Corruption), sizeof(Corruption));
+		FileHandle.Reset();
+
+		Store = MakeUnique<FAnalyticsPropertyStore>();
+		check(!Store->Load(TestStorePathname)); // Checksum should fails.
+	}
+
+	// Test store internal reset.
+	{
+		// Create the store.
+		FAnalyticsPropertyStore Store;
+		check(Store.Create(TestStorePathname, /*CapacityHint*/2 * 1024));
+		Store.Set(I32Key, Expected_I32);
+		Store.Flush();
+
+		// Reload the store -> This will reset the object.
+		check(Store.Load(TestStorePathname));
+		check(Store.Get(I32Key, Actual_I32) == IAnalyticsPropertyStore::EStatusCode::Success && Expected_I32 == Actual_I32);
+
+		// Recreate the store -> This will reset the object and create an empty store.
+		check(Store.Create(TestStorePathname, /*CapacityHint*/2 * 1024));
+		check(Store.Get(I32Key, Actual_I32) == IAnalyticsPropertyStore::EStatusCode::NotFound);
+		check(!Store.Contains(I32Key));
+	}
+
+	// Test flushing async - no timeout
+	{
+		// Create the store.
+		FAnalyticsPropertyStore Store;
+		check(Store.Create(TestStorePathname, /*CapacityHint*/2 * 1024));
+		Store.Set(I32Key, Expected_I32);
+		Store.Flush(/*bAsync*/true);
+
+		// Reload the store -> This will wait until the async task complates and reset the object.
+		check(Store.Load(TestStorePathname));
+		check(Store.Get(I32Key, Actual_I32) == IAnalyticsPropertyStore::EStatusCode::Success && Expected_I32 == Actual_I32);
+	}
+
+	// Test flushing async - with timeout
+	{
+		// Create the store.
+		FAnalyticsPropertyStore Store;
+		check(Store.Create(TestStorePathname, /*CapacityHint*/2 * 1024));
+		Store.Set(I32Key, Expected_I32);
+		Store.Flush(/*bAsync*/true, FTimespan::FromMilliseconds(100));
+
+		// Reload the store -> This will wait until the async task complates and reset the object.
+		check(Store.Load(TestStorePathname));
+		check(Store.Get(I32Key, Actual_I32) == IAnalyticsPropertyStore::EStatusCode::Success && Expected_I32 == Actual_I32);
+	}
+
+
+	IFileManager::Get().Delete(*TestStorePathname);
+
+	return true;
+}
+
+#endif // WITH_DEV_AUTOMATION_TESTS
diff --git a/Engine/Source/Runtime/Analytics/Analytics/Public/AnalyticsPropertyStore.h b/Engine/Source/Runtime/Analytics/Analytics/Public/AnalyticsPropertyStore.h
new file mode 100644
--- /dev/null
+++ b/Engine/Source/Runtime/Analytics/Analytics/Public/AnalyticsPropertyStore.h
@@ -0,0 +1,182 @@
+// Copyright Epic Games, Inc. All Rights Reserved.
+
+#pragma once
+
+#include "CoreMinimal.h"
+#include "Async/AsyncWork.h"
+#include "Interfaces/IAnalyticsPropertyStore.h"
+#include "Misc/DateTime.h"
+#include "Serialization/Archive.h"
+#include "Serialization/MemoryReader.h"
+#include "Serialization/MemoryWriter.h"
+#include "Templates/Function.h"
+
+class IFileHandle;
+
+/**
+ * Implements a fast key/value database to store the metrics gathered to create the session summary event.
+ * The store is buffered in memory until it is flushed to disk by calling Flush(). It is not designed to
+ * handle millions of key/value pairs, but rather a small subset that fits in memory. The implementation
+ * reserve disk space ahead of time and can usually perform in-place update. Setting or getting intrinsic
+ * types (int32, uint32, int64, uint64, float, double, bool as well as FDateTime) doesn't allocate memory.
+ * The string time is often updated in place if enough capacity was reserved. All store operations are
+ * thread safe and atomic from the caller perspective.
+ */
+class ANALYTICS_API FAnalyticsPropertyStore : public IAnalyticsPropertyStore
+{
+public:
+	/** Create an empty store. */
+	FAnalyticsPropertyStore();
+	virtual ~FAnalyticsPropertyStore();
+
+	FAnalyticsPropertyStore(const FAnalyticsPropertyStore&) = delete;
+	FAnalyticsPropertyStore(FAnalyticsPropertyStore&&) = delete;
+	FAnalyticsPropertyStore& operator=(const FAnalyticsPropertyStore&) = delete;
+	FAnalyticsPropertyStore& operator=(FAnalyticsPropertyStore&&) = delete;
+
+	/**
+	 * Create a new store, creating the file of the specified capacity and resetting the current state.
+	 * @param Pathname The path to the file to create.
+	 * @param CapacityHint The desired file capacity. The implementation has a minimum and maximum capacity to reserve.
+	 * @param true if the file was created and the store created, false otherwise.
+	 */
+	bool Create(const FString& Pathname, uint32 CapacityHint = 4 * 1024);
+
+	/**
+	 * Load the specified file in memory, resetting the current state.
+	 * @param Pathname The path to the file to open.
+	 * @return true if the file was loaded successuflly, false if the file doesn't exist, the file format/version is invalid or if the file is corrupted (failed checksum).
+	 */
+	bool Load(const FString& Pathname);
+
+	/** Returns true is the store was successfully created or loaded. */
+	bool IsValid() const { return FileHandle != nullptr; }
+
+	//~ IAnalyticsPropertyStore interface
+	virtual uint32 Num() const override;
+	virtual bool Contains(const FString& Key) const override;
+	virtual bool Remove(const FString& Key) override;
+	virtual void RemoveAll() override;
+	virtual EStatusCode Set(const FString& Key, int32 Value) override;
+	virtual EStatusCode Set(const FString& Key, uint32 Value) override;
+	virtual EStatusCode Set(const FString& Key, int64 Value) override;
+	virtual EStatusCode Set(const FString& Key, uint64 Value) override;
+	virtual EStatusCode Set(const FString& Key, float Value) override;
+	virtual EStatusCode Set(const FString& Key, double Value) override;
+	virtual EStatusCode Set(const FString& Key, bool Value) override;
+	virtual EStatusCode Set(const FString& Key, const FString& Value, uint32 CharCountCapacityHint = 0) override;
+	virtual EStatusCode Set(const FString& Key, const FDateTime& Value) override;
+	virtual EStatusCode Set(const FString& Key, int32 Value, const TFunction<bool(const int32* /*Actual*/, const int32& /*Proposed*/)>& ConditionFn) override;
+	virtual EStatusCode Set(const FString& Key, uint32 Value, const TFunction<bool(const uint32* /*Actual*/, const uint32& /*Proposed*/)>& ConditionFn) override;
+	virtual EStatusCode Set(const FString& Key, int64 Value, const TFunction<bool(const int64* /*Actual*/, const int64& /*Proposed*/)>& ConditionFn) override;
+	virtual EStatusCode Set(const FString& Key, uint64 Value, const TFunction<bool(const uint64* /*Actual*/, const uint64& /*Proposed*/)>& ConditionFn) override;
+	virtual EStatusCode Set(const FString& Key, float Value, const TFunction<bool(const float* /*Actual*/, const float& /*Proposed*/)>& ConditionFn) override;
+	virtual EStatusCode Set(const FString& Key, double Value, const TFunction<bool(const double* /*Actual*/, const double& /*Proposed*/)>& ConditionFn) override;
+	virtual EStatusCode Set(const FString& Key, bool Value, const TFunction<bool(const bool*/*Actual*/, const bool& /*Proposed*/)>& ConditionFn) override;
+	virtual EStatusCode Set(const FString& Key, const FString& Value, const TFunction<bool(const FString* /*Actual*/, const FString& /*Proposed*/)>& ConditionFn) override;
+	virtual EStatusCode Set(const FString& Key, const FDateTime& Value, const TFunction<bool(const FDateTime* /*Actual*/, const FDateTime& /*Proposed*/)>& ConditionFn) override;
+	virtual EStatusCode Update(const FString& Key, const TFunction<bool(int32&)>& UpdateFn) override;
+	virtual EStatusCode Update(const FString& Key, const TFunction<bool(uint32&)>& UpdateFn) override;
+	virtual EStatusCode Update(const FString& Key, const TFunction<bool(int64&)>& UpdateFn) override;
+	virtual EStatusCode Update(const FString& Key, const TFunction<bool(uint64&)>& UpdateFn) override;
+	virtual EStatusCode Update(const FString& Key, const TFunction<bool(float&)>& UpdateFn) override;
+	virtual EStatusCode Update(const FString& Key, const TFunction<bool(double&)>& UpdateFn) override;
+	virtual EStatusCode Update(const FString& Key, const TFunction<bool(bool&)>& UpdateFn) override;
+	virtual EStatusCode Update(const FString& Key, const TFunction<bool(FString&)>& UpdateFn) override;
+	virtual EStatusCode Update(const FString& Key, const TFunction<bool(FDateTime&)>& UpdateFn) override;
+	virtual EStatusCode Get(const FString& Key, int32& OutValue) const override;
+	virtual EStatusCode Get(const FString& Key, uint32& OutValue) const override;
+	virtual EStatusCode Get(const FString& Key, int64& OutValue) const override;
+	virtual EStatusCode Get(const FString& Key, uint64& OutValue) const override;
+	virtual EStatusCode Get(const FString& Key, float& OutValue) const override;
+	virtual EStatusCode Get(const FString& Key, double& OutValue) const override;
+	virtual EStatusCode Get(const FString& Key, bool& OutValue) const override;
+	virtual EStatusCode Get(const FString& Key, FString& OutValue) const override;
+	virtual EStatusCode Get(const FString& Key, FDateTime& OutValue) const override;
+	virtual void VisitAll(const TFunction<void(const FString& /*Key*/, FString&& /*Value*/)>& VisitFn) const override;
+
+	/**
+	 * Flushes cached values to disk. Only one thread can flush at a time, causing possible latency if a flush is alreary happenning.
+	 * @param bAsync Whether the data in flushed in a background thread or in the calling thread.
+	 * @param Timeout Maximum time to wait before the flush starts (synchronous) or get scheduled (asynchronous). FTimespan::Zero() means no waiting.
+	 * @return True if the flush was executed (synchronous) or scheduled (asynchronous), false if the operation timed out.
+	 *
+	 * @code
+	 *   Store.Flush(); // Flush in the calling thread, block until all threads has finished flushing.
+	 *   Store.Flush(false, FTimespan::Zero()); // Flush in the calling thread if possible, return immediatedly if another thread is already flushing.
+	 *   Store.Flush(true, FTimespan::Zero()); // Flush in a backgroud thread if no other thread is flushing, return immediatedly if another flushing task is running.
+	 *   Store.Flush(true);  // Flush in a background thread, might block if a thread is already flushing. (It doesn't queue several tasks).
+	 * @endcode
+	 */
+	virtual bool Flush(bool bAsync = false, const FTimespan& Timeout = FTimespan::MaxValue()) override;
+
+private:
+	enum class ETypeCode : uint8
+	{
+		I32     = 0x00,
+		U32     = 0x01,
+		I64     = 0x02,
+		U64     = 0x03,
+		Flt     = 0x04,
+		Dbl     = 0x05,
+		Bool    = 0x06,
+		Date    = 0x07,
+		Str     = 0x08,
+		RawMask = 0x0F, // To extract the raw types above from the meta data below
+		Dead    = 0x80, // Type for record that were moved or deleted.
+	};
+	FRIEND_ENUM_CLASS_FLAGS(ETypeCode)
+
+	/** Helper task to persist the properties to disk in the background. */
+	class FFlushWorker : public FNonAbandonableTask
+	{
+	public:
+		FFlushWorker(FAnalyticsPropertyStore& InStore);
+		void DoWork();
+		TStatId GetStatId() const { return TStatId(); }
+		static const TCHAR* Name() { return TEXT("AnalyticsPropertyStoreWorker"); }
+	private:
+		FAnalyticsPropertyStore& Store;
+		TArray<uint8> Data;
+	};
+	friend class FFlushWorker;
+
+private:
+	EStatusCode SetFixedSizeValueInternal(const FString& Key, ETypeCode TypeCode, const uint8* Value, uint32 Size);
+	EStatusCode GetFixedSizeValueInternal(const FString& Key, ETypeCode TypeCode, uint8* OutValue, uint32 Size) const;
+	EStatusCode SetStringValueInternal(const FString& Key, const FString& Value, uint32 CharCountCapacityHint = 0);
+	EStatusCode GetStringValueInternal(const FString& Key, FString& OutValue) const;
+	EStatusCode SetDateTimeValueInternal(const FString& Key, const FDateTime& Value);
+	EStatusCode GetDateTimeValueInternal(const FString& Key, FDateTime& OutValue) const;
+	void Defragment();
+	bool IsFixedSize(ETypeCode Type) const { return Type != ETypeCode::Str; }
+	bool IsDead(ETypeCode Type) const { return EnumHasAnyFlags(Type, ETypeCode::Dead); }
+	ETypeCode RawType(ETypeCode Type) const { return Type & ETypeCode::RawMask; }
+	void FlushInternal(const TArray<uint8>& Data);
+	void Reset();
+
+	template<typename T>
+	IAnalyticsPropertyStore::EStatusCode SetFixedSizeValueInternal(const FString& Key, ETypeCode TypeCode, T Value)
+	{
+		return SetFixedSizeValueInternal(Key, TypeCode, reinterpret_cast<const uint8*>(&Value), sizeof(Value));
+	}
+
+	template<typename T>
+	EStatusCode GetFixedSizeValueInternal(const FString& Key, ETypeCode TypeCode, T& OutValue) const
+	{
+		return GetFixedSizeValueInternal(Key, TypeCode, reinterpret_cast<uint8*>(&OutValue), sizeof(OutValue));
+	}
+
+private:
+	mutable FCriticalSection StoreLock;
+	TArray<uint8> StorageBuf;
+	FMemoryWriter StorageWriter;
+	mutable FMemoryReader StorageReader;
+	TMap<FString, uint32> NameOffsetMap;
+	TUniquePtr<IFileHandle> FileHandle;
+	bool bFragmented = false;
+	TUniquePtr<FAsyncTask<FFlushWorker>> FlushTask;
+	TArray<uint8> AsyncFlushDataCopy;
+};
+
+ENUM_CLASS_FLAGS(FAnalyticsPropertyStore::ETypeCode);
diff --git a/Engine/Source/Runtime/Analytics/Analytics/Public/Interfaces/IAnalyticsPropertyStore.h b/Engine/Source/Runtime/Analytics/Analytics/Public/Interfaces/IAnalyticsPropertyStore.h
new file mode 100644
--- /dev/null
+++ b/Engine/Source/Runtime/Analytics/Analytics/Public/Interfaces/IAnalyticsPropertyStore.h
@@ -0,0 +1,183 @@
+// Copyright Epic Games, Inc. All Rights Reserved.
+
+#pragma once
+
+#include "CoreMinimal.h"
+#include "Misc/DateTime.h"
+#include "Templates/Function.h"
+
+/**
+ * Implements a fast type safe key/value database to store analytics properties collected during a session.
+ */
+class ANALYTICS_API IAnalyticsPropertyStore
+{
+public:
+	/** The list of store operation outcome. */
+	enum class EStatusCode : uint32
+	{
+		/** The operation terminated sucessfully. */
+		Success,
+		/** The operation failed because the key could not be found. */
+		NotFound,
+		/** The operation failed because the key value provided did not match the currently stored value. */
+		BadType,
+		/** A conditional operation failed because it was declined by the caller (during a callback). */
+		Declined,
+	};
+
+public:
+	virtual ~IAnalyticsPropertyStore() = default;
+
+	/**
+	 * Returns the number of elements in the store.
+	 */
+	virtual uint32 Num() const = 0;
+
+	/**
+	 * Returns true if the store contains the specified key.
+	 */
+	virtual bool Contains(const FString& Key) const = 0;
+
+	/**
+	 * Removes the specified key from the store.
+	 * @return true if the key was removed.
+	 */
+	virtual bool Remove(const FString& Key) = 0;
+
+	/**
+	 * Removes all existing keys from the store.
+	 */
+	virtual void RemoveAll() = 0;
+
+	/**
+	 * Adds or updates the specified key/value pair. If the key doesn't exist, the key/value pair is added. If the key already exists, the value is updated if the stored value and the
+	 * specified value types match.
+	 * @param Key The property name to add or update.
+	 * @param Value The property value to add or update.
+	 * @return EStatusCode::Success if the key was added or updated, EStatusCode::BadType if the value type did not match the current value type.
+	 */
+	virtual EStatusCode Set(const FString& Key, int32 Value) = 0;
+	virtual EStatusCode Set(const FString& Key, uint32 Value) = 0;
+	virtual EStatusCode Set(const FString& Key, int64 Value) = 0;
+	virtual EStatusCode Set(const FString& Key, uint64 Value) = 0;
+	virtual EStatusCode Set(const FString& Key, float Value) = 0;
+	virtual EStatusCode Set(const FString& Key, double Value) = 0;
+	virtual EStatusCode Set(const FString& Key, bool Value) = 0;
+	virtual EStatusCode Set(const FString& Key, const FString& Value, uint32 CharCountCapacityHint = 0) = 0;
+	virtual EStatusCode Set(const FString& Key, const FDateTime& Value) = 0;
+
+	/**
+	 * Conditionnnaly adds or updates the specified key/value pair. The operation is atomic from the caller perspective. The function reads the actual value (if it exists), invoke the ConditionFn
+	 * callback and write the value back if the condition function returns true. If the key doesn't exist, the actual value passed back to the condition function is null. If the condition function
+	 * returns true, the proposed value is set, otherwise, the operation is declined and the store remains unchanged.
+	 * @param Key The name of the property to add or update.
+	 * @param Value The value proposed. This value is passed as the second parameter to the ConditionFn callback.
+	 * @param ConditionFn Function invoked back with the actual value and the proposed value to let the caller decide whether the proposed value should be set or not.
+	 *                    If the callback function returns true, the proposed value is added/updated, otherwise, nothing changes.
+	 * @return One of the following code
+	 *   - EStatusCode::Success if the entire operation succeeded.
+	 *   - EStatusCode::BadType if the key/value pair exists, but the proposed value did not match the actual value type.
+	 *   - EStatusCode::Declined if the callback ConditionFn returned false and the operation was declined.
+	 */
+	virtual EStatusCode Set(const FString& Key, int32 Value, const TFunction<bool(const int32* /*Actual*/, const int32& /*Proposed*/)>& ConditionFn) = 0;
+	virtual EStatusCode Set(const FString& Key, uint32 Value, const TFunction<bool(const uint32* /*Actual*/, const uint32& /*Proposed*/)>& ConditionFn) = 0;
+	virtual EStatusCode Set(const FString& Key, int64 Value, const TFunction<bool(const int64* /*Actual*/, const int64& /*Proposed*/)>& ConditionFn) = 0;
+	virtual EStatusCode Set(const FString& Key, uint64 Value, const TFunction<bool(const uint64* /*Actual*/, const uint64& /*Proposed*/)>& ConditionFn) = 0;
+	virtual EStatusCode Set(const FString& Key, float Value, const TFunction<bool(const float* /*Actual*/, const float& /*Proposed*/)>& ConditionFn) = 0;
+	virtual EStatusCode Set(const FString& Key, double Value, const TFunction<bool(const double* /*Actual*/, const double& /*Proposed*/)>& ConditionFn) = 0;
+	virtual EStatusCode Set(const FString& Key, bool Value, const TFunction<bool(const bool*/*Actual*/, const bool& /*Proposed*/)>& ConditionFn) = 0;
+	virtual EStatusCode Set(const FString& Key, const FString& Value, const TFunction<bool(const FString* /*Actual*/, const FString& /*Proposed*/)>& ConditionFn) = 0;
+	virtual EStatusCode Set(const FString& Key, const FDateTime& Value, const TFunction<bool(const FDateTime* /*Actual*/, const FDateTime& /*Proposed*/)>& ConditionFn) = 0;
+
+	/**
+	 * Updates an exiting value. The operation is atomic from the caller perspective. The function reads the actual value and invoke UpdateFn callback with the actual value. The
+	 * callback can update the actual value or decline the operation. If the key is not found, the callback is not invoked.
+	 * @param Key The name of the key to update.
+	 * @param UpdateFn The function to invoke, passing the actual value.
+	 * @return One of the following code
+	 *   - EStatusCode::Success if the entire operation succeeded.
+	 *   - EStatusCode::BadType if the key/value pair exists, but the proposed value did not match the actual value type.
+	 *   - EStatusCode::Declined if the callback UpdateFn returned false and the operation was declined.
+	 */
+	virtual EStatusCode Update(const FString& Key, const TFunction<bool(int32& /*InOutValue*/)>& UpdateFn) = 0;
+	virtual EStatusCode Update(const FString& Key, const TFunction<bool(uint32& /*InOutValue*/)>& UpdateFn) = 0;
+	virtual EStatusCode Update(const FString& Key, const TFunction<bool(int64& /*InOutValue*/)>& UpdateFn) = 0;
+	virtual EStatusCode Update(const FString& Key, const TFunction<bool(uint64& /*InOutValue*/)>& UpdateFn) = 0;
+	virtual EStatusCode Update(const FString& Key, const TFunction<bool(float& /*InOutValue*/)>& UpdateFn) = 0;
+	virtual EStatusCode Update(const FString& Key, const TFunction<bool(double& /*InOutValue*/)>& UpdateFn) = 0;
+	virtual EStatusCode Update(const FString& Key, const TFunction<bool(bool& /*InOutValue*/)>& UpdateFn) = 0;
+	virtual EStatusCode Update(const FString& Key, const TFunction<bool(FString& /*InOutValue*/)>& UpdateFn) = 0;
+	virtual EStatusCode Update(const FString& Key, const TFunction<bool(FDateTime& /*InOutValue*/)>& UpdateFn) = 0;
+
+	/**
+	 * Reads the specified key value from the store.
+	 * @param Key The name of the property to read.
+	 * @param OutValue The value of the property if the function returns EStatusCode::Success, left unchanged otherwise.
+	 * @return One of the following code
+	 *   - EStatusCode::Success if the value was read succeessfully.
+	 *   - EStatusCode::BadType if the key/value pair exists, but the type used did not match the actual value type.
+	 *   - EStatusCode::NotFound if the key/value could not be found.
+	 */
+	virtual EStatusCode Get(const FString& Key, int32& OutValue) const = 0;
+	virtual EStatusCode Get(const FString& Key, uint32& OutValue) const = 0;
+	virtual EStatusCode Get(const FString& Key, int64& OutValue) const = 0;
+	virtual EStatusCode Get(const FString& Key, uint64& OutValue) const = 0;
+	virtual EStatusCode Get(const FString& Key, float& OutValue) const = 0;
+	virtual EStatusCode Get(const FString& Key, double& OutValue) const = 0;
+	virtual EStatusCode Get(const FString& Key, bool& OutValue) const = 0;
+	virtual EStatusCode Get(const FString& Key, FString& OutValue) const = 0;
+	virtual EStatusCode Get(const FString& Key, FDateTime& OutValue)const = 0;
+
+	/**
+	 * Iterates the keys currently stored and invokes the visitor function for each key, converting the value to its string representation.
+	 * @param VisitFn The callback invoked for each key/value pair visited.
+	 */
+	virtual void VisitAll(const TFunction<void(const FString& /*Key*/, FString&& /*Value*/)>& VisitFn) const = 0;
+
+	/**
+	 * Flushes cached values to persistent storage.
+	 * @param bAsync Whether the data in flushed in a background thread or in the calling thread.
+	 * @param Timeout Maximum time to wait before the flush starts (synchronous) or get scheduled (asynchronous). FTimespan::Zero() means no waiting.
+	 * @return True if the flush was executed (synchronous) or scheduled (asynchronous), false if the operation timed out.
+	 */
+	virtual bool Flush(bool bAsync = false, const FTimespan& Timeout = FTimespan::MaxValue()) = 0;
+};
+
+/**
+ * Utility class to remember the key value type and let the compiler implicitly convert type or fail if such conversion is not possible.
+ * Usage:
+ * @code
+ *   static const TAnalyticsProperty<uint32> MyProperty = TEXT("MyProperty");
+ *   IAnalyticsPropertyStore* MyStore = ...;
+ *   MyProperty.SetValue(MyStore,  10); // Ok.
+ *   MyProperty.SetValue(MyStore, TEXT("No")); // Compile error, not convertible.
+ * @endcode
+ */
+template<typename T>
+class TAnalyticsProperty
+{
+public:
+	TAnalyticsProperty(const TCHAR* InKey) : Key(InKey){}
+
+	IAnalyticsPropertyStore::EStatusCode Set(IAnalyticsPropertyStore* Store, const T& Value) const { return Store->Set(Key, Value); }
+	IAnalyticsPropertyStore::EStatusCode Set(IAnalyticsPropertyStore* Store, const T& Value, const TFunction<bool(const T*, const T&)>& ConditionFn) const { return Store->Set(Key, Value, ConditionFn); }
+	IAnalyticsPropertyStore::EStatusCode Get(IAnalyticsPropertyStore* Store, T& OutValue) const { return Store->Get(Key, OutValue); }
+	IAnalyticsPropertyStore::EStatusCode Update(IAnalyticsPropertyStore* Store, const TFunction<bool(T& InOutValue)>& UpdateFn) const { return Store->Update(Key, UpdateFn); }
+
+	const FString Key;
+};
+
+/** Specialization for FString because the caller can specifies a capacity when adding a string. */
+template<>
+class TAnalyticsProperty<FString>
+{
+public:
+	TAnalyticsProperty(const TCHAR* InKey) : Key(InKey){}
+
+	IAnalyticsPropertyStore::EStatusCode Set(IAnalyticsPropertyStore* Store, const FString& Value, uint32 CharCountCapacityHint = 0) const { return Store->Set(Key, Value, CharCountCapacityHint); }
+	IAnalyticsPropertyStore::EStatusCode Set(IAnalyticsPropertyStore* Store, const FString& Value, const TFunction<bool(const FString*, const FString&)>& ConditionFn) const { return Store->Set(Key, Value, ConditionFn); }
+	IAnalyticsPropertyStore::EStatusCode Get(IAnalyticsPropertyStore* Store, FString& OutValue) const { return Store->Get(Key, OutValue); }
+	IAnalyticsPropertyStore::EStatusCode Update(IAnalyticsPropertyStore* Store, const TFunction<bool(FString& InOutValue)>& UpdateFn) const { return Store->Update(Key, UpdateFn); }
+
+	const FString Key;
+};
diff --git a/Engine/Source/Runtime/Analytics/AnalyticsET/Private/AnalyticsSessionSummaryManager.cpp b/Engine/Source/Runtime/Analytics/AnalyticsET/Private/AnalyticsSessionSummaryManager.cpp
new file mode 100644
--- /dev/null
+++ b/Engine/Source/Runtime/Analytics/AnalyticsET/Private/AnalyticsSessionSummaryManager.cpp
@@ -0,0 +1,593 @@
+// Copyright Epic Games, Inc. All Rights Reserved.
+
+#include "AnalyticsSessionSummaryManager.h"
+#include "AnalyticsPropertyStore.h"
+#include "IAnalyticsSessionSummarySender.h"
+#include "GenericPlatform/GenericPlatformProcess.h"
+#include "HAL/CriticalSection.h"
+#include "HAL/FileManager.h"
+#include "HAL/PlatformProcess.h"
+#include "HAL/PlatformTime.h"
+#include "Misc/App.h"
+#include "Misc/Optional.h"
+#include "Internationalization/Regex.h"
+
+DEFINE_LOG_CATEGORY_STATIC(LogAnalyticsSessionSummary, Verbose, All);
+
+namespace AnalyticsManagerProperties
+{
+	// Those values are added by the manager for internal usage. They are stripped before sending the summary.
+	static const FString InternalPropertyPrefix                        = TEXT("Internal.");
+	static const TAnalyticsProperty<FString> InternalSessionUserId     = TEXT("Internal.UserId");
+	static const TAnalyticsProperty<FString> InternalSessionAppId      = TEXT("Internal.AppId");
+	static const TAnalyticsProperty<FString> InternalSessionAppVersion = TEXT("Internal.AppVersion");
+	static const TAnalyticsProperty<FString> InternalSessionId         = TEXT("Internal.SessionId");
+	static const TAnalyticsProperty<bool>    InternalWasProcessed      = TEXT("Internal.WasProcessed");
+
+	// Those values are implicitely added by the manager into the report because the analytics backend expecpts them.
+	static const TAnalyticsProperty<FString> ShutdownType      = TEXT("ShutdownType");
+	static const TAnalyticsProperty<bool>    DelayedSend       = TEXT("DelayedSend");
+	static const TAnalyticsProperty<FString> SentFrom          = TEXT("SentFrom");
+
+	// List of reserved property key.
+	static TSet<FString> ReservedKeys = {
+			AnalyticsManagerProperties::ShutdownType.Key,
+			AnalyticsManagerProperties::DelayedSend.Key,
+			AnalyticsManagerProperties::SentFrom.Key
+		};
+
+	bool IsReserved(const FString& Key)
+	{
+		return ReservedKeys.Contains(Key);
+	}
+} // namespace AnalyticsManagerProperties
+
+
+namespace AnalyticsSessionUtils
+{
+	static const FTimespan OrphanSessionCheckPeriod = FTimespan::FromMinutes(5);
+	static const FTimespan ProcessGroupDeathGracePeriod = FTimespan::FromHours(1);
+	static const FTimespan ProcessGroupExpiration = FTimespan::FromDays(30.0);
+
+	static const FString OrphanSessionsOwnerCriticalSectionName = TEXT("UE_AnalyticsSessionSummaryManager_OrphanOwner");
+	static const FString OrphanSessionsOwnerFilename            = TEXT("8E1D46DBC38F4A789939D781E1B91520"); // A randomly generated GUID.
+
+	FString GetDefaultSavedDir()
+	{
+		// NOTE: This needs to be a folder common to Epic applications (Editor, CRC, etc.), ideally agnostics from the engine versions to enable picking orphans from
+		//       older versions AND being constant between build configurations (like Editor Developement and CRC Shipping). FPaths functions are sensitive to the
+		//       build configuration and will put files in different places.
+		return FPaths::Combine(FPlatformProcess::UserSettingsDir(), *FApp::GetEpicProductIdentifier(), TEXT("Common"), TEXT("Analytics"));
+	}
+
+	FRegexPattern GetAnalyticsPropertyFilenamePattern()
+	{
+		// Filename format is like "ProcessGroupGUID_PID_PID_ProcessName", for example 3E1D7ADBC38F4A789939D781E1B91520_12345_232513_Editor"
+		return FRegexPattern(TEXT(R"((^[0-9A-Za-z]+)_([0-9]+)_([0-9]+)_(.*))")); // Need help with regex? Try https://regex101.com/
+	}
+
+	/** Try to acquire the rigth to process and send sessions for orphan groups. */
+	bool TryAcquiringOrphanGroupsOwnership(const FString& SessionRootPath, const FString& PrincipalProcessName, uint32 CurrentProcessId)
+	{
+		// NOTE: The usage of FPlatformMisc::SetStoredValue()/GetStoredValue() was considered, but those APIs are not thread safe on Mac/Linux
+		//       and concurrent processes can easily stomp the changes made by another one. The implementation below prevent races between threads and processes.
+		FAnalyticsPropertyStore OwnerStore;
+		FString OrphanOwnerFilePathname = SessionRootPath / AnalyticsSessionUtils::OrphanSessionsOwnerFilename;
+
+		bool bAcquiredOwnership = false;
+
+		// Try to lock. Do not wait if the lock is already taken. That's not a big deal if we delay processing left-over sessions.
+		FSystemWideCriticalSection SysWideLock(AnalyticsSessionUtils::OrphanSessionsOwnerCriticalSectionName, FTimespan::Zero());
+		if (SysWideLock.IsValid())
+		{
+			if (IFileManager::Get().FileExists(*OrphanOwnerFilePathname))
+			{
+				if (OwnerStore.Load(OrphanOwnerFilePathname))
+				{
+					IAnalyticsPropertyStore::EStatusCode Status = OwnerStore.Set(PrincipalProcessName, CurrentProcessId,
+						[](const uint32* Actual, const int32& ProposedValue) { return Actual == nullptr || *Actual == 0 || !FPlatformProcess::IsApplicationRunning(*Actual); });
+
+					bAcquiredOwnership = (Status == IAnalyticsPropertyStore::EStatusCode::Success);
+				}
+				else // If the file exists but fails to load while this process as the system wide lock, that's likely because its content is invalid/corrupted.
+				{
+					IFileManager::Get().Delete(*OrphanOwnerFilePathname);
+				}
+			}
+			else if (OwnerStore.Create(OrphanOwnerFilePathname, /*Capacity*/1 * 1024))
+			{
+				bAcquiredOwnership = OwnerStore.Set(PrincipalProcessName, CurrentProcessId) == IAnalyticsPropertyStore::EStatusCode::Success;
+			}
+
+			if (bAcquiredOwnership)
+			{
+				OwnerStore.Flush(); // Ensure to persist the storage to disk.
+			}
+		}
+
+		return bAcquiredOwnership;
+	}
+
+	/** Release the ownership to send sessions for orphan groups. */
+	void ReleaseOrphanGroupsOwnership(const FString& SessionRootPath, const FString& PrincipalProcessName, uint32 CurrentProcessId)
+	{
+		FString OrphanOwnerFilePathname = SessionRootPath / AnalyticsSessionUtils::OrphanSessionsOwnerFilename;
+		if (IFileManager::Get().FileExists(*OrphanOwnerFilePathname))
+		{
+			// Try to lock. Do not wait if the lock is already taken.
+			FSystemWideCriticalSection SysWideLock(AnalyticsSessionUtils::OrphanSessionsOwnerCriticalSectionName, FTimespan::Zero());
+			if (SysWideLock.IsValid())
+			{
+				// Open the file.
+				FAnalyticsPropertyStore OwnerStore;
+				if (OwnerStore.Load(OrphanOwnerFilePathname))
+				{
+					// Reset to zero for this principal process.
+					IAnalyticsPropertyStore::EStatusCode Status = OwnerStore.Set(PrincipalProcessName, static_cast<uint32>(0),
+						[CurrentProcessId](const uint32* Actual, const uint32& Proposed) { return Actual == nullptr || *Actual == CurrentProcessId; });
+
+					if (Status == IAnalyticsPropertyStore::EStatusCode::Success)
+					{
+						OwnerStore.Flush();
+					}
+				}
+			}
+		}
+	}
+
+	/** Convert the shutdown type to the string representation known by the analytics backend. */
+	const TCHAR* ToString(EAnalyticsSessionShutdownType Code)
+	{
+		switch (Code)
+		{
+		case EAnalyticsSessionShutdownType::Shutdown:
+			return TEXT("Shutdown");
+
+		case EAnalyticsSessionShutdownType::Crashed:
+			return TEXT("Crashed");
+
+		case EAnalyticsSessionShutdownType::Terminated:
+			return TEXT("Terminated");
+
+		case EAnalyticsSessionShutdownType::Debugged:
+			return TEXT("Debugger");
+
+		case EAnalyticsSessionShutdownType::Abnormal:
+			return TEXT("AbnormalShutdown");
+
+		default:
+			return TEXT("Unknown");
+		}
+	}
+
+} // namespace AnalyticsSessionUtils
+
+
+const TAnalyticsProperty<int32> FAnalyticsSessionSummaryManager::ShutdownTypeCodeProperty = TEXT("ShutdownTypeCode");
+const TAnalyticsProperty<bool>  FAnalyticsSessionSummaryManager::IsUserLoggingOutProperty = TEXT("IsUserLoggingOut");
+
+FAnalyticsSessionSummaryManager::FAnalyticsSessionSummaryManager(const FString& InProcessName, const FString& InProcessGroupId,const FString& InUserId, const FString& InAppId, const FString& InAppVersion, const FString& InSessionId, const FString& InSavedDir)
+	: FAnalyticsSessionSummaryManager(InProcessName, InProcessGroupId, FPlatformProcess::GetCurrentProcessId(), FPlatformProcess::GetCurrentProcessId(), InUserId, InAppId, InAppVersion, InSessionId, InSavedDir)
+{
+}
+
+FAnalyticsSessionSummaryManager::FAnalyticsSessionSummaryManager(const FString& InProcessName, const FString& InProcessGroupId, uint32 InPrincipalProcessId, const FString& InSavedDir)
+	: FAnalyticsSessionSummaryManager(InProcessName, InProcessGroupId, FPlatformProcess::GetCurrentProcessId(), InPrincipalProcessId, FString(), FString(), FString(), FString(), InSavedDir)
+{
+}
+
+FAnalyticsSessionSummaryManager::FAnalyticsSessionSummaryManager(const FString& InProcessName, const FString& InProcessGroupId, uint32 InCurrentProcessId, uint32 InPrincipalProcessId, const FString& InUserId, const FString& InAppId, const FString& InAppVersion, const FString& InSessionId, const FString& InSavedDir)
+	: ProcessName(InProcessName)
+	, ProcessGroupId(InProcessGroupId)
+	, UserId(InUserId)
+	, AppId(InAppId)
+	, AppVersion(InAppVersion)
+	, SessionId(InSessionId)
+	, SessionRootPath(InSavedDir)
+	, CurrentProcessId(InCurrentProcessId)
+	, PrincipalProcessId(InPrincipalProcessId)
+	, NextOrphanSessionCheckTimeSecs(FPlatformTime::Seconds())
+	, bOrphanGroupOwner(false)
+	, bIsPrincipal(InCurrentProcessId == InPrincipalProcessId)
+{
+	if (SessionRootPath.IsEmpty() || !IFileManager::Get().DirectoryExists(*SessionRootPath))
+	{
+		SessionRootPath = AnalyticsSessionUtils::GetDefaultSavedDir(); // Fallback on the default.
+	}
+
+	// Creates a sub-folder to store the session files.
+	IFileManager::Get().MakeDirectory(*SessionRootPath, /*Tree*/true);
+
+	// Let principal process deals with its orphans. The principal process is expected to knows its application settings while a subsidiary process may not have any context.
+	if (IsPrincipalProcess())
+	{
+		bOrphanGroupOwner = AnalyticsSessionUtils::TryAcquiringOrphanGroupsOwnership(SessionRootPath, ProcessName, CurrentProcessId);
+	}
+}
+
+FAnalyticsSessionSummaryManager::~FAnalyticsSessionSummaryManager()
+{
+	if (bOrphanGroupOwner)
+	{
+		AnalyticsSessionUtils::ReleaseOrphanGroupsOwnership(SessionRootPath, ProcessName, CurrentProcessId);
+	}
+}
+
+TSharedPtr<IAnalyticsPropertyStore> FAnalyticsSessionSummaryManager::MakeStore(uint32 InitialCapacity)
+{
+	// NOTE: The ProcessGroupId (a GUID) + PrincipalProcessId is expected to provide a unique key that should not collide with existing files.
+	FString PropertyStorePathname = SessionRootPath / FString::Printf(TEXT("%s_%d_%d_%s"), *ProcessGroupId, PrincipalProcessId, CurrentProcessId, *ProcessName);
+
+	TSharedPtr<FAnalyticsPropertyStore> PropertyStore = MakeShared<FAnalyticsPropertyStore>();
+	if (PropertyStore && PropertyStore->Create(PropertyStorePathname, InitialCapacity))
+	{
+		// If this is the principal process for which the session summary is created.
+		if (IsPrincipalProcess())
+		{
+			// Add some internal key required to enable a subsidiary process to sent the summary on the behalf of this process.
+			AnalyticsManagerProperties::InternalSessionUserId.Set(PropertyStore.Get(), UserId);
+			AnalyticsManagerProperties::InternalSessionAppId.Set(PropertyStore.Get(), AppId);
+			AnalyticsManagerProperties::InternalSessionAppVersion.Set(PropertyStore.Get(), AppVersion);
+			AnalyticsManagerProperties::InternalSessionId.Set(PropertyStore.Get(), SessionId);
+		}
+	}
+	else
+	{
+		PropertyStore.Reset(); // Ensure to return a null pointer if the store failed to create (not enough disk space for example)
+	}
+
+	return StaticCastSharedPtr<IAnalyticsPropertyStore>(PropertyStore);
+}
+
+void FAnalyticsSessionSummaryManager::SetSender(TSharedPtr<IAnalyticsSessionSummarySender> Sender)
+{
+	SummarySender = Sender;
+}
+
+void FAnalyticsSessionSummaryManager::Shutdown(bool bDiscard)
+{
+	if (SummarySender)
+	{
+		// Check if this process has something to summarize and send for this instance process group.
+		TMap<FString, FProcessGroup> PropertyFiles = GetSessionFiles();
+		if (const FProcessGroup* ProcessGroup = PropertyFiles.Find(ProcessGroupId))
+		{
+			if (bDiscard)
+			{
+				CleanupFiles(ProcessGroup->PropertyFiles, /*bOnSuccess*/false);
+			}
+			else
+			{
+				ProcessSummary(ProcessGroupId, *ProcessGroup);
+			}
+		}
+	}
+	else if (bDiscard)
+	{
+		TMap<FString, FProcessGroup> PropertyFiles = GetSessionFiles();
+		if (const FProcessGroup* ProcessGroup = PropertyFiles.Find(ProcessGroupId))
+		{
+			CleanupFiles(ProcessGroup->PropertyFiles, /*bOnSuccess*/false);
+		}
+	}
+}
+
+void FAnalyticsSessionSummaryManager::Tick()
+{
+	if (!SummarySender || !IsOrphanGroupsOwner())
+	{
+		return; // Don't process orphans if the result cannot be sent.
+	}
+
+	if (FPlatformTime::Seconds() > NextOrphanSessionCheckTimeSecs)
+	{
+		NextOrphanSessionCheckTimeSecs = FPlatformTime::Seconds() + AnalyticsSessionUtils::OrphanSessionCheckPeriod.GetTotalSeconds();
+
+		// Check if this process has something to summarize and send..
+		TMap<FString, FProcessGroup> PropertyFiles = GetSessionFiles();
+		for (const TPair<FString, FProcessGroup>& Pair : PropertyFiles)
+		{
+			if (Pair.Value.PrincipalProcessId == PrincipalProcessId)
+			{
+				continue; // Skip this process session summary. It should not be sent until the last process in the group died.
+			}
+
+			// Picked up an orphan session, decide what to do with it.
+			ProcessSummary(/*ProcessGroupId*/Pair.Key, /*ProcessGroup*/Pair.Value);
+		}
+	}
+}
+
+bool FAnalyticsSessionSummaryManager::IsPrincipalProcess() const
+{
+	return bIsPrincipal;
+}
+
+TMap<FString, FAnalyticsSessionSummaryManager::FProcessGroup> FAnalyticsSessionSummaryManager::GetSessionFiles() const
+{
+	TMap<FString, FProcessGroup> ProcessFiles;
+
+	// Visit the files in the session directory.
+	FRegexPattern FilenamePattern = AnalyticsSessionUtils::GetAnalyticsPropertyFilenamePattern();
+	IFileManager::Get().IterateDirectory(*SessionRootPath, [this, &ProcessFiles, &FilenamePattern](const TCHAR* Pathname, bool bIsDirectory) -> bool
+	{
+		if (!bIsDirectory)
+		{
+			// The property filename is "ProcessGroupGUID_PrincipalProcessId_WriterProcessId_WriterProcessName". The ProcessGroupGUID is expected to have EGuidFormats::Digits -> "00000A00C0000F000000000000000000"
+			FRegexMatcher Matcher(FilenamePattern, FPaths::GetCleanFilename(Pathname));
+			if (Matcher.FindNext())
+			{
+				FString FileProcessGroupId = Matcher.GetCaptureGroup(1);
+				if (FileProcessGroupId != ProcessGroupId && !IsOrphanGroupsOwner())
+				{
+					return true; // This file is from another process group and this manager is not allowed to process it, skip to the next file.
+				}
+
+				FPropertyFileInfo PropertyFileInfo;
+				PropertyFileInfo.ProcessId = static_cast<uint32>(FCString::Atoi64(*Matcher.GetCaptureGroup(3)));
+				PropertyFileInfo.ProcessName = Matcher.GetCaptureGroup(4);
+				PropertyFileInfo.Pathname = Pathname;
+
+				uint32 FilePrincipalProcessId = static_cast<uint32>(FCString::Atoi64(*Matcher.GetCaptureGroup(2)));
+				FProcessGroup& ProcessGroup = ProcessFiles.FindOrAdd(FileProcessGroupId);
+				ProcessGroup.PrincipalProcessId = FilePrincipalProcessId;
+
+				if (FilePrincipalProcessId == PropertyFileInfo.ProcessId) // That file was produced by a principal process.
+				{
+					ProcessGroup.PrincipalProcessName = PropertyFileInfo.ProcessName;
+				}
+
+				ProcessGroup.PropertyFiles.Add(MoveTemp(PropertyFileInfo));
+			}
+		}
+		return true; // Continue to the next file.
+	});
+
+	// Can this manager process orphan groups?
+	if (IsOrphanGroupsOwner())
+	{
+		check(IsPrincipalProcess());
+
+		for (TMap<FString, FProcessGroup>::TIterator It = ProcessFiles.CreateIterator(); It; ++It)
+		{
+			// Discard groups that are not from this applications.
+			if (!It.Value().PrincipalProcessName.IsEmpty() && It.Value().PrincipalProcessName != ProcessName)
+			{
+				It.RemoveCurrent();
+			}
+			// Discard any other groups that are still running.
+			else if (It.Value().PrincipalProcessId != PrincipalProcessId)
+			{
+				for (const FPropertyFileInfo& PropertyFileInfo : It.Value().PropertyFiles)
+				{
+					if (FPlatformProcess::IsApplicationRunning(PropertyFileInfo.ProcessId)) // One of the process in that group is still running
+					{
+						It.RemoveCurrent();
+						break;
+					}
+				}
+			}
+		}
+	}
+
+	return ProcessFiles;
+}
+
+void FAnalyticsSessionSummaryManager::ProcessSummary(const FString& InProcessGroupId, const FProcessGroup& InProcessGroup)
+{
+	// Check if this process is the last alive in the group, so that it can process the summary
+	for (const FPropertyFileInfo& PropertyFileInfo : InProcessGroup.PropertyFiles)
+	{
+		if (PropertyFileInfo.ProcessId != CurrentProcessId)
+		{
+			if (FPlatformProcess::IsApplicationRunning(PropertyFileInfo.ProcessId))
+			{
+				return; // A process in the group is running, leave this group alone.
+			}
+		}
+	}
+
+	// Trying to process a summary for another group than this manager group?
+	if (InProcessGroupId != ProcessGroupId)
+	{
+		checkf(IsOrphanGroupsOwner(), TEXT("Should not process summary of other process group if this wasn't allowed in first place."));
+
+		bool bExpired = false;
+
+		// The files discovered tell that processes in the group are dead, but it is unclear if another process in that group is about to create a files or not (because
+		// of race conditions between processes).
+		for (const FPropertyFileInfo& PropertyFileInfo : InProcessGroup.PropertyFiles)
+		{
+			// Time since the last file modification.
+			double FileAgeSecs = IFileManager::Get().GetFileAgeSeconds(*PropertyFileInfo.Pathname);
+
+			// Before declaring this group 'dead' leave a reasonable grace period to clear 99% of possible race conditions between processes.
+			if (FileAgeSecs > 0 && FileAgeSecs <= AnalyticsSessionUtils::ProcessGroupDeathGracePeriod.GetTotalSeconds())
+			{
+				return; // Wait more before declaring this group dead and processing it.
+			}
+
+			// If the principal file is expired, expire the entire group.
+			if (InProcessGroup.PrincipalProcessId == PropertyFileInfo.ProcessId && FileAgeSecs > AnalyticsSessionUtils::ProcessGroupExpiration.GetTotalSeconds())
+			{
+				bExpired = true;
+				break;
+			}
+		}
+
+		// This group is expired or doesn't have the properties from the principal process (incomplete session).
+		if (bExpired || InProcessGroup.PrincipalProcessName.IsEmpty())
+		{
+			// The session data is partial or expired.
+			CleanupFiles(InProcessGroup.PropertyFiles, /*bOnSuccess*/false);
+			return;
+		}
+	}
+
+	// Aggregate the principal and subsidiary process summaries.
+	TMap<FString, FString> InternalProperties;
+	TMap<FString, FString> SummaryProperties;
+	if (!AggregateSummaries(InProcessGroupId, InProcessGroup.PropertyFiles, SummaryProperties, InternalProperties))
+	{
+		CleanupFiles(InProcessGroup.PropertyFiles, /*bOnSuccess*/false);
+		return;
+	}
+
+	// Was the group aggregated above the group of this principal process?
+	if (IsPrincipalProcess() && InProcessGroupId == ProcessGroupId)
+	{
+		// Clean up before sending, so that we don't send twice. (In case of crash while sending, we will not know if the summary was sent or not).
+		CleanupFiles(InProcessGroup.PropertyFiles, /*bOnSuccess*/true);
+		SummarySender->SendSessionSummary(UserId, AppId, AppVersion, SessionId, SummaryProperties);
+	}
+	// Send the session summary on behalf of the principal process that created it.
+	else
+	{
+		// The principal process saves those internal properties to let a subsidiary or another principal processes submit the summary on its behalf.
+		const FString* UserIdProp = InternalProperties.Find(AnalyticsManagerProperties::InternalSessionUserId.Key);
+		const FString* AppIdProp = InternalProperties.Find(AnalyticsManagerProperties::InternalSessionAppId.Key);
+		const FString* AppVersionProp = InternalProperties.Find(AnalyticsManagerProperties::InternalSessionAppVersion.Key);
+		const FString* SessionIdProp = InternalProperties.Find(AnalyticsManagerProperties::InternalSessionId.Key);
+
+		if (UserIdProp && AppIdProp && AppVersionProp && SessionIdProp)
+		{
+			CleanupFiles(InProcessGroup.PropertyFiles, /*bOnSuccess*/true);
+			SummarySender->SendSessionSummary(*UserIdProp, *AppIdProp, *AppVersionProp, *SessionIdProp, SummaryProperties);
+		}
+		else
+		{
+			CleanupFiles(InProcessGroup.PropertyFiles, /*bOnSuccess*/false);
+		}
+	}
+}
+
+bool FAnalyticsSessionSummaryManager::AggregateSummaries(const FString& InProcessGroupId, const TArray<FPropertyFileInfo>& PropertyFiles, TMap<FString, FString>& OutSummaryProperties, TMap<FString, FString>& OutInternalProperties)
+{
+	FAnalyticsPropertyStore PropertyStore;
+
+	auto RenameKey = [](const FString& KeyToRename, const FString& FromProcess, const TMap<FString, FString>& ExistingKeys, const TCHAR* Reason)
+	{
+		FString RenamedKey;
+		int32 Index = 1;
+		do
+		{
+			RenamedKey = FString::Printf(TEXT("%s_%d"), *KeyToRename, Index++);
+		} while (ExistingKeys.Contains(RenamedKey));
+
+		UE_LOG(LogAnalyticsSessionSummary, Warning, TEXT("Renamed key '%s' from process '%s' as '%s'. Reason: %s"), *KeyToRename, *FromProcess, *RenamedKey, Reason);
+		return RenamedKey;
+	};
+
+	// Merge the properties collected by the principal and its subsidiary processes.
+	for (const FPropertyFileInfo& PropertyFileInfo : PropertyFiles)
+	{
+		if (PropertyStore.Load(PropertyFileInfo.Pathname))
+		{
+			if (PropertyStore.Contains(AnalyticsManagerProperties::InternalWasProcessed.Key))
+			{
+				return false; // We don't know if the manager failed to delete before of after the session was sent. Ignore this group.
+			}
+
+			int32 ShutdownTypeCode;
+			if (PropertyStore.Get(ShutdownTypeCodeProperty.Key, ShutdownTypeCode) == IAnalyticsPropertyStore::EStatusCode::Success)
+			{
+				// Don't add 'ShutdownType' more than once. If 'ShutdownTypeCode' key used to derive the 'ShutdownType' is duplicated, the code will emit a warning and rename the second instance.
+				if (!OutSummaryProperties.Contains(AnalyticsManagerProperties::ShutdownType.Key))
+				{
+					// Convert 'ShutdownTypeCode' into its string representation and emit it as 'ShutdownType' as known by the analytics backend.
+					OutSummaryProperties.Emplace(AnalyticsManagerProperties::ShutdownType.Key, AnalyticsSessionUtils::ToString(static_cast<EAnalyticsSessionShutdownType>(ShutdownTypeCode)));
+				}
+			}
+
+			PropertyStore.VisitAll([&PropertyFileInfo, &OutSummaryProperties, &OutInternalProperties, &RenameKey](const FString& InKey, FString&& InValue)
+			{
+				if (InKey.StartsWith(AnalyticsManagerProperties::InternalPropertyPrefix))
+				{
+					OutInternalProperties.Emplace(InKey, InValue);
+				}
+				else if (AnalyticsManagerProperties::IsReserved(InKey))
+				{
+					FString RenamedKey = RenameKey(InKey, PropertyFileInfo.ProcessName, OutSummaryProperties, TEXT("Key name is reserved."));
+					OutSummaryProperties.Emplace(RenamedKey, InValue);
+				}
+				else if (!OutSummaryProperties.Contains(InKey))
+				{
+					OutSummaryProperties.Emplace(InKey, InValue);
+				}
+				else // Deal with key collisions. Keys are supposed to be unique. Duplication should be easy to spot during developement.
+				{
+					FString RenamedKey = RenameKey(InKey, PropertyFileInfo.ProcessName, OutSummaryProperties, TEXT("Key collision between processes of a group."));
+					OutSummaryProperties.Emplace(RenamedKey, InValue);
+				}
+			});
+
+			// Mark this file as processed to avoid sending the same session twice.
+			AnalyticsManagerProperties::InternalWasProcessed.Set(&PropertyStore, true);
+			PropertyStore.Flush();
+		}
+		else
+		{
+			return false; // The file failed to load. Was it deleted? locked? internal checksum failed?
+		}
+	}
+
+	// The manager saves this key in the store it creates.
+	if (const FString* InternalSessionId = OutInternalProperties.Find(AnalyticsManagerProperties::InternalSessionId.Key))
+	{
+		OutSummaryProperties.Emplace(TEXT("SessionId"), *InternalSessionId);
+	}
+
+	// The summary is sent delayed when it is processed from another group.
+	OutSummaryProperties.Emplace(AnalyticsManagerProperties::DelayedSend.Key, LexToString(InProcessGroupId != ProcessGroupId));
+
+	// This process is about to send the report, record the process name in the summary.
+	OutSummaryProperties.Emplace(AnalyticsManagerProperties::SentFrom.Key, ProcessName);
+
+	return true;
+}
+
+bool FAnalyticsSessionSummaryManager::CleanupFiles(const TArray<FPropertyFileInfo>& PropertyFiles, bool bOnSuccess)
+{
+	// TODO as future work: Compute how many times we clean up on failure. Keep track of it and periodically send an analytic event about it.
+
+	bool bAllDeleted = true;
+	for (const FPropertyFileInfo& PropertyFileInfo : PropertyFiles)
+	{
+		bAllDeleted &= IFileManager::Get().Delete(*PropertyFileInfo.Pathname, /*MustExist*/false, /*EvenReadOnly*/true, /*Quiet*/true);
+	}
+	return bAllDeleted;
+}
+
+void FAnalyticsSessionSummaryManager::CleanupExpiredFiles(const FString& SavedDir)
+{
+	FString SessionRootPath = SavedDir.IsEmpty() ? AnalyticsSessionUtils::GetDefaultSavedDir() : SavedDir;
+
+	// The manager stores its files under its own folder.
+	if (!IFileManager::Get().DirectoryExists(*SessionRootPath))
+	{
+		return;
+	}
+
+	// Visit the files in the session directory.
+	FRegexPattern FilenamePattern = AnalyticsSessionUtils::GetAnalyticsPropertyFilenamePattern();
+	IFileManager::Get().IterateDirectory(*SessionRootPath, [&FilenamePattern](const TCHAR* Pathname, bool bIsDirectory) -> bool
+	{
+		if (bIsDirectory)
+		{
+			return true; // Continue.
+		}
+
+		FRegexMatcher Matcher(FilenamePattern, FPaths::GetCleanFilename(Pathname));
+		if (Matcher.FindNext())
+		{
+			// Time since the last file modification.
+			double FileAgeSecs = IFileManager::Get().GetFileAgeSeconds(Pathname);
+			if (FileAgeSecs > AnalyticsSessionUtils::ProcessGroupExpiration.GetTotalSeconds() + FTimespan::FromDays(1.0).GetTotalSeconds())
+			{
+				IFileManager::Get().Delete(Pathname);
+			}
+		}
+
+		return true;
+	});
+}
diff --git a/Engine/Source/Runtime/Analytics/AnalyticsET/Private/AnalyticsSessionSummarySender.cpp b/Engine/Source/Runtime/Analytics/AnalyticsET/Private/AnalyticsSessionSummarySender.cpp
new file mode 100644
--- /dev/null
+++ b/Engine/Source/Runtime/Analytics/AnalyticsET/Private/AnalyticsSessionSummarySender.cpp
@@ -0,0 +1,75 @@
+// Copyright Epic Games, Inc. All Rights Reserved.
+
+#include "AnalyticsSessionSummarySender.h"
+#include "AnalyticsET.h"
+#include "AnalyticsEventAttribute.h"
+#include "IAnalyticsProviderET.h"
+#include "Misc/EngineVersion.h"
+
+DEFINE_LOG_CATEGORY_STATIC(LogAnalyticsSessionSummarySender, Verbose, All);
+
+
+FAnalyticsSessionSummarySender::FAnalyticsSessionSummarySender(IAnalyticsProviderET& Provider, TFunction<bool(const FString&, const FString&)> ShouldEmitFilterFunc)
+	: AnalyticsProvider(Provider)
+	, ShouldEmitPropFunc(MoveTemp(ShouldEmitFilterFunc))
+{
+}
+
+bool FAnalyticsSessionSummarySender::SendSessionSummary(const FString& UserId, const FString& AppId, const FString& AppVersion, const FString& SessionId, const TMap<FString, FString>& Properties)
+{
+	TArray<FAnalyticsEventAttribute> AnalyticsAttributes;
+
+	// Track which app/user is sending the summary (summary can be sent by another process (CrashReportClient) or another instance in case of crash/abnormal terminaison.
+	AnalyticsAttributes.Emplace(TEXT("SummarySenderAppId"), AnalyticsProvider.GetAppID());
+	AnalyticsAttributes.Emplace(TEXT("SummarySenderAppVersion"), AnalyticsProvider.GetAppVersion());
+	AnalyticsAttributes.Emplace(TEXT("SummarySenderEngineVersion"), FEngineVersion::Current().ToString(EVersionComponent::Changelist)); // Same as in EditorSessionSummaryWriter.cpp
+	AnalyticsAttributes.Emplace(TEXT("SummarySenderUserId"), AnalyticsProvider.GetUserID());
+	AnalyticsAttributes.Emplace(TEXT("SummarySenderSessionId"), AnalyticsProvider.GetSessionID()); // Not stripping the {} around the GUID like EditorSessionSummaryWriter does with SessionId.
+
+	for (const TPair<FString, FString>& Property : Properties)
+	{
+		// Function is not bound or returns true.
+		if (!ShouldEmitPropFunc || ShouldEmitPropFunc(Property.Key, Property.Value))
+		{
+			AnalyticsAttributes.Emplace(Property.Key, Property.Value);
+		}
+	}
+
+	// Sort the keys to makes it easier to find specific ones when inspecting the network payload (mainly for testing/debugging purpose).
+	AnalyticsAttributes.Sort([](const FAnalyticsEventAttribute& Lhs, const FAnalyticsEventAttribute& Rhs){ return Lhs.GetName() < Rhs.GetName(); });
+
+	// Sending the summary event of the current process analytic session?
+	if (AnalyticsProvider.GetSessionID().Contains(SessionId)) // The string (GUID) returned by GetSessionID() is surrounded with braces like "{3FEA3232-...}" while Session.SessionId is not -> "3FEA3232-..."
+	{
+		AnalyticsProvider.RecordEvent(TEXT("SessionSummary"), AnalyticsAttributes);
+	}
+	else // The summary was created by another process/instance in a different session. (Ex: Editor sending a summary of a  previoulsy crashed instance or CrashReportClient sending it on behalf of the Editor)
+	{
+		// The provider sending a 'summary event' created by another instance/process must parametrize its post request 'as if' it was sent from the instance/session that created it (backend expectation).
+		// Create a new provider to avoid interfering with the current session events. (ex. if another thread sends telemetry at the same time, don't accidently tag it with the wrong SessionID, AppID, etc.).
+		TSharedPtr<IAnalyticsProviderET> TempSummaryProvider = FAnalyticsET::Get().CreateAnalyticsProvider(AnalyticsProvider.GetConfig());
+
+		// Reconfigure the analytics provider to sent the summary event 'as if' it was sent by the process that created it. This is required by the analytics backend.
+		FGuid SessionGuid;
+		FGuid::Parse(*SessionId, SessionGuid);
+		TempSummaryProvider->SetUserID(CopyTemp(*UserId));
+		TempSummaryProvider->SetAppID(CopyTemp(*AppId));
+		TempSummaryProvider->SetAppVersion(CopyTemp(*AppVersion));
+		TempSummaryProvider->SetSessionID(SessionGuid.ToString(EGuidFormats::DigitsWithHyphensInBraces)); // Ensure to put back the {} around the GUID.
+
+		// Send the summary.
+		TempSummaryProvider->RecordEvent(TEXT("SessionSummary"), AnalyticsAttributes);
+
+		// The temporary provider is about to be deleted (going out of scope), ensure it sents its report.
+		TempSummaryProvider->BlockUntilFlushed(2.0f);
+	}
+
+	// For debugging.
+	//for (const FAnalyticsEventAttribute& Att : AnalyticsAttributes)
+	//{
+	//	UE_LOG(LogAnalyticsSessionSummarySender, Display, TEXT("Attr %s = %s"), *Att.GetName(), *Att.GetValue());
+	//}
+
+	UE_LOG(LogAnalyticsSessionSummarySender, Log, TEXT("Sent summary session report for: AppId=%s SessionId=%s"), *AppId, *SessionId);
+	return true;
+}
diff --git a/Engine/Source/Runtime/Analytics/AnalyticsET/Private/Tests/AnalyticsSessionSummaryManagerTests.cpp b/Engine/Source/Runtime/Analytics/AnalyticsET/Private/Tests/AnalyticsSessionSummaryManagerTests.cpp
new file mode 100644
--- /dev/null
+++ b/Engine/Source/Runtime/Analytics/AnalyticsET/Private/Tests/AnalyticsSessionSummaryManagerTests.cpp
@@ -0,0 +1,138 @@
+// Copyright Epic Games, Inc. All Rights Reserved.
+
+#include "CoreMinimal.h"
+#include "GenericPlatform/GenericPlatformFile.h"
+#include "HAL/PlatformFileManager.h"
+#include "HAL/PlatformProcess.h"
+#include "Interfaces/IAnalyticsPropertyStore.h"
+#include "Misc/AutomationTest.h"
+#include "AnalyticsSessionSummaryManager.h"
+#include "IAnalyticsSessionSummarySender.h"
+#include "Misc/Paths.h"
+
+//#if WITH_DEV_AUTOMATION_TESTS
+
+IMPLEMENT_SIMPLE_AUTOMATION_TEST(FAnalyticsSessionSummaryManagerAutomationTest, "System.Analytics.SessionSummaryManager", EAutomationTestFlags::EditorContext | EAutomationTestFlags::EngineFilter | EAutomationTestFlags::MediumPriority)
+
+// A dummy sender to observe what the manager would send.
+class FDummyAnalyticsSummarySender : public IAnalyticsSessionSummarySender
+{
+public:
+	virtual bool SendSessionSummary(const FString& InUserId, const FString& InAppId, const FString& InAppVersion, const FString& InSessionId, const TMap<FString, FString>& InProperties) override
+	{
+		UserId = InUserId;
+		AppId = InAppId;
+		AppVersion = InAppVersion;
+		SessionId = InSessionId;
+		Properties = InProperties;
+		bReportSent = true;
+		return true;
+	}
+
+	bool bReportSent = false;
+	FString UserId;
+	FString AppId;
+	FString AppVersion;
+	FString SessionId;
+	TMap<FString, FString> Properties;
+};
+
+bool FAnalyticsSessionSummaryManagerAutomationTest::RunTest(const FString& Parameters)
+{
+	const FString PrincipalProcessName       = TEXT("DummyPrincipal");
+	const FString SubsidiaryProcessName      = TEXT("DummySubsidiary");
+	const FString ProcessGroupId             = FGuid::NewGuid().ToString(EGuidFormats::Digits);
+	const FString PrincipalProcessUserId     = TEXT("DummyUserId");
+	const FString PrincipalProcessAppId      = TEXT("DummyAppId");
+	const FString PrincipalProcessAppVersion = TEXT("DummyAppVersion");
+	const FString PrincipalProcessSessionId  = FGuid::NewGuid().ToString(EGuidFormats::DigitsWithHyphensInBraces);
+	const FString SavedDir                   = FPaths::AutomationTransientDir() / TEXT("AnalyticsTest");
+
+	IFileManager::Get().MakeDirectory(*SavedDir, /*Tree*/true);
+
+	// Properties automatically added by the manager.
+	const TAnalyticsProperty<bool>    DelayedSendProp = TEXT("DelayedSend");
+	const TAnalyticsProperty<FString> SessionIdProp   = TEXT("SessionId");
+	const TAnalyticsProperty<FString> SentFromProp    = TEXT("SentFrom");
+
+	const TAnalyticsProperty<int32> Prop1 = TEXT("Prop1");
+	const int32 Prop1Value = 42;
+
+	const TAnalyticsProperty<int64> Prop2 = TEXT("Prop2");
+	const int32 Prop2Value = 84;
+
+	//
+	// Test the basic flow simulating a principal process running alone (Like Editor without CRC).
+	//
+	{
+		// Create the manager.
+		FAnalyticsSessionSummaryManager PrincipalManager(PrincipalProcessName, ProcessGroupId, PrincipalProcessUserId, PrincipalProcessAppId, PrincipalProcessAppVersion, PrincipalProcessSessionId, SavedDir);
+
+		// Set the summary sender to observe what the manager is going to send.
+		TSharedPtr<FDummyAnalyticsSummarySender> SummarySender = MakeShared<FDummyAnalyticsSummarySender>();
+		PrincipalManager.SetSender(StaticCastSharedPtr<IAnalyticsSessionSummarySender>(SummarySender));
+
+		// Create a key/value database to store the session summary properties.
+		TSharedPtr<IAnalyticsPropertyStore> PropertyStore = PrincipalManager.MakeStore(/*Capacity*/1024);
+		check(PropertyStore);
+
+		// Store a property, flush and release the store (so that it can be reloaded).
+		Prop1.Set(PropertyStore.Get(), Prop1Value);
+		PropertyStore->Flush();
+		PropertyStore.Reset();
+
+		// Shutdown the manager. This should emits the summary.
+		PrincipalManager.Shutdown();
+
+		// Ensure the manager used the correct config to send.
+		check(SummarySender->UserId     == PrincipalProcessUserId);
+		check(SummarySender->AppId      == PrincipalProcessAppId);
+		check(SummarySender->AppVersion == PrincipalProcessAppVersion);
+		check(SummarySender->SessionId  == PrincipalProcessSessionId);
+
+		// Check the automatically embedded keys.
+		check(*SummarySender->Properties.Find(DelayedSendProp.Key)  == LexToString(false));
+		check(*SummarySender->Properties.Find(SessionIdProp.Key)    == PrincipalProcessSessionId);
+		check(*SummarySender->Properties.Find(SentFromProp.Key)     == PrincipalProcessName);
+
+		// Check the custom property.
+		check(*SummarySender->Properties.Find(Prop1.Key) == LexToString(Prop1Value));
+	}
+
+	//
+	// Test discarding the session in a principal process.
+	//
+	{
+		// Create the managers.
+		FAnalyticsSessionSummaryManager PrincipalManager(PrincipalProcessName, ProcessGroupId, PrincipalProcessUserId, PrincipalProcessAppId, PrincipalProcessAppVersion, PrincipalProcessSessionId, SavedDir);
+		
+		// Set the summary sender to observe what the managers are going to send.
+		TSharedPtr<FDummyAnalyticsSummarySender> PrincipalSender = MakeShared<FDummyAnalyticsSummarySender>();
+		PrincipalManager.SetSender(StaticCastSharedPtr<IAnalyticsSessionSummarySender>(PrincipalSender));
+
+		// Create a key/value database to store the session summary properties.
+		TSharedPtr<IAnalyticsPropertyStore> PrincipalStore = PrincipalManager.MakeStore(/*Capacity*/1024);
+		check(PrincipalStore);
+
+		// Store a property.
+		Prop1.Set(PrincipalStore.Get(), Prop1Value);
+
+		// Flush and release the store.
+		PrincipalStore->Flush();
+		PrincipalStore.Reset();
+
+		// Discard the summary
+		PrincipalManager.Shutdown(/*bDiscard*/true);
+
+		// Ensure no report was sent.
+		check(!PrincipalSender->bReportSent);
+
+		// Ensure the file created by the process was deleted.
+		FString PropertyStorePathname = SavedDir / FString::Printf(TEXT("%s_%d_%d_%s"), *ProcessGroupId, FPlatformProcess::GetCurrentProcessId(), FPlatformProcess::GetCurrentProcessId(), *PrincipalProcessName);
+		check(!IFileManager::Get().FileExists(*PropertyStorePathname));
+	}
+
+	return true;
+}
+
+//#endif WITH_DEV_AUTOMATION_TESTS
diff --git a/Engine/Source/Runtime/Analytics/AnalyticsET/Public/AnalyticsSessionSummaryManager.h b/Engine/Source/Runtime/Analytics/AnalyticsET/Public/AnalyticsSessionSummaryManager.h
new file mode 100644
--- /dev/null
+++ b/Engine/Source/Runtime/Analytics/AnalyticsET/Public/AnalyticsSessionSummaryManager.h
@@ -0,0 +1,171 @@
+// Copyright Epic Games, Inc. All Rights Reserved.
+
+#pragma once
+
+#include "CoreMinimal.h"
+#include "Interfaces/IAnalyticsPropertyStore.h"
+
+class IAnalyticsSessionSummarySender;
+
+/** Defines how the principal process (the application for which analytics is gathered) exited. When the analytics session summary manager aggregates the summaries for a session,
+    it looks for 'ShutdownTypeCode' key and if the key is found, it converts it to its string representation and add the 'ShutdownType' property known by the analytics backend. */
+enum class EAnalyticsSessionShutdownType : int32
+{
+	Shutdown,
+	Terminated,
+	Debugged,
+	Crashed,
+	Abnormal,
+	Unknown,
+};
+
+/**
+ * Manages the lifecycle of analytics summaries created by an application and its subsidiary processes. It manages the
+ * summary files created by the processes, tracks which processes are working together, aggregates the summaries
+ * from collaborating processes, discards outdated sessions and sends the final summary reports.
+ *
+ * Several processes can collaborate to collect analytics for one principal process. The common case being the Editor acting as
+ * the principal process and CrashReportClient (CRC) acting as a subsidiary process. The principal process and its subsidiaries
+ * form a group of processes tied together. The last process of the group to exit is responsible to aggregates and sends the final
+ * report for the principal process. For any reasons, if the final report cannot be sent by a group, the responsibility to deal
+ * with the left over is put on the principal process. The principal process is responsible to clean up the left over from its
+ * previous execution(s) and possibly salvage and send reports delayed.
+ */
+class ANALYTICSET_API FAnalyticsSessionSummaryManager
+{
+public:
+	/** Special key handled by the manager for the analytics backend. Its value must be one of EAnalyticsSessionShutdownType enum value. */
+	static const TAnalyticsProperty<int32> ShutdownTypeCodeProperty;
+
+	/** Special key handled in CRC to decide if an abnormal shutdown crash report should be generated. */
+	static const TAnalyticsProperty<bool> IsUserLoggingOutProperty;
+
+public:
+	/**
+	 * Constructs a manager for the principal process, usually the main application.
+	 * @param ProcessName The process tag name of the principal process. Ex. "Editor".
+	 * @param ProcessGroupId A unique ID shared by the principal and subsidiary processes. This is how information created by various processes is associated.
+	 * @param UserId The current analytics user ID. @see IAnalyticsProviderET
+	 * @param AppId The current application ID used by analytics to identify the application. @see IAnalyticsProviderET
+	 * @param AppVersion the current application version used by the analytics for this application. @see IAnalyticsProviderET
+	 * @param SessionId the current session ID assigned by the analytics. @see IAnalyticsProviderET
+	 * @param SavedDir The directory where files will be persisted. Should be the same for the principal and the subsidiary processes. If unspecified, the manager uses its internal default.
+	 */
+	FAnalyticsSessionSummaryManager(const FString& ProcessName, const FString& ProcessGroupId, const FString& UserId, const FString& AppId, const FString& AppVersion, const FString& SessionId, const FString& SavedDir = TEXT(""));
+
+	/**
+	 * Constructs a manager for a subsidiary process, usually a companion process that collects extra information on the behalf of the principal process.
+	 * @param ProcessName The process tag name of the subsidiary process. Ex. "CrashReportClient".
+	 * @param ProcessGroupId A unique ID shared by the principal and subsidiary processes. This is how information created by various processes is associated.
+	 * @param PrincipalProcessId The principal process ID to which this subsidiary process is associated.
+	 * @param SavedDir The directory where files will be persisted. Should be the same for the principal and the subsidiary processes. If unspecified, the manager uses its internal default.
+	 */
+	FAnalyticsSessionSummaryManager(const FString& ProcessName, const FString& ProcessGroupId, uint32 PrincipalProcessId, const FString& SavedDir = TEXT(""));
+
+	/**
+	 * Destructor.
+	 */
+	~FAnalyticsSessionSummaryManager();
+
+	/**
+	 * Discovers and salvages/sends/cleans up left over from previous execution(s), if any. Only the principal process can send left over sessions.
+	 */
+	void Tick();
+
+	/**
+	 * Creates a new property store associated to this manager process group. The manager only sends an analytics session summary when all
+	 * collecting processes have closed their property store. Ensure to flush and release the store before calling the manager Shutdown() to
+	 * ensure the session is sent as soon a possible.
+	 * @param InitialCapacity The amount of space to reserve in the file.
+	 */
+	TSharedPtr<IAnalyticsPropertyStore> MakeStore(uint32 InitialCapacity);
+
+	/**
+	 * Sets a summary sender to enable the manager to send sessions. If the sender is null, the sessions are not sent. If the
+	 * sender is set, the manager will periodically search and send pending sessions, usually the one left over from
+	 * processes that died unexpectedly. As a rule, the principal process is responsible to send the left over summary sessions
+	 * and a subsidiary process only sends the current sessions summary if it is the last process of the group to exit. The manager
+	 * deletes the session data while sending.
+	 */
+	void SetSender(TSharedPtr<IAnalyticsSessionSummarySender> Sender);
+
+	/**
+	 * Shuts down and sends the session summaries that can be sent, if a sender is set. Remember to flush and discard the
+	 * property store returned by MakeStore() before shutting down to prevent delaying the report.
+	 * @param bDiscard Whether the current session is discarded or not. Discarded sessions are not sent.
+	 * @see SetSender()
+	 * @see MakeStore()
+	 */
+	void Shutdown(bool bDiscard = false);
+
+	/**
+	 * Cleans up all expired files that match the analytics session filename pattern. The manager automaticallly clean up if
+	 * it runs, but if the analytics is off and the manager is not instantiated anymore, some dead files can be left over.
+	 * @param SavedDir The directory used to save the analytic files. If unspecified, use the internal default.
+	 */
+	static void CleanupExpiredFiles(const FString& SavedDir = TEXT(""));
+
+private:
+	/** Information about a property store used by a process. */
+	struct FPropertyFileInfo
+	{
+		/** The ID of the process that wrote the property file. Ex. CRC can write a summary for an Editor session. */
+		uint32 ProcessId = 0;
+
+		/** The tag name of the process that wrote the property file. */
+		FString ProcessName;
+
+		/** Pathname to the property file. */
+		FString Pathname;
+	};
+
+	/** Lists the files created by a process group. */
+	struct FProcessGroup
+	{
+		/** Identify the process for which the summary was written. */
+		uint32 PrincipalProcessId = 0;
+
+		/** The principal process name. */
+		FString PrincipalProcessName;
+
+		/** The list of files generated by the processes in the group. */
+		TArray<FPropertyFileInfo> PropertyFiles;
+	};
+
+private:
+	/** Delegate constructor. */
+	FAnalyticsSessionSummaryManager(const FString& ProcessName, const FString& ProcessGroupId, uint32 InCurrentProcessId, uint32 PrincipalProcessId, const FString& UserId, const FString& AppId, const FString& AppVersion, const FString& SessionId, const FString& SessionRootDir);
+
+	/** Returns all processes group (key is process group id) for which files exists. */
+	TMap<FString, FProcessGroup> GetSessionFiles() const;
+
+	/** Processes the information collected by the principal and its subsidiary processes and decides if a summmary can be aggregated and sent. */
+	void ProcessSummary(const FString& ProcessGroupId, const FProcessGroup& ProcessGroup);
+
+	/** Loads the summaries created for the application and its subsidiary processes (like Editor/CRC combo) and aggregates summary properties. */
+	bool AggregateSummaries(const FString& ProcessGroupId, const TArray<FPropertyFileInfo>& PropertyFiles, TMap<FString, FString>& OutSummaryProperties, TMap<FString, FString>& OutInternalProperties);
+
+	/** Deletes the files used to record and perist the properties. */
+	bool CleanupFiles(const TArray<FPropertyFileInfo>& PropertyFiles, bool bOnSuccess);
+
+	/** Checks if this process is allowed to process and send orphan files. */
+	bool IsOrphanGroupsOwner() const { return bOrphanGroupOwner; }
+
+	/** Returns whether this manager instance represents the principal process of a group. */
+	bool IsPrincipalProcess() const;
+
+private:
+	TSharedPtr<IAnalyticsSessionSummarySender> SummarySender;
+	FString ProcessName;
+	FString ProcessGroupId;
+	FString UserId;
+	FString AppId;
+	FString AppVersion;
+	FString SessionId;
+	FString SessionRootPath;
+	uint32 CurrentProcessId;
+	uint32 PrincipalProcessId;
+	double NextOrphanSessionCheckTimeSecs;
+	bool bOrphanGroupOwner = false;
+	bool bIsPrincipal = false;
+};
diff --git a/Engine/Source/Runtime/Analytics/AnalyticsET/Public/AnalyticsSessionSummarySender.h b/Engine/Source/Runtime/Analytics/AnalyticsET/Public/AnalyticsSessionSummarySender.h
new file mode 100644
--- /dev/null
+++ b/Engine/Source/Runtime/Analytics/AnalyticsET/Public/AnalyticsSessionSummarySender.h
@@ -0,0 +1,36 @@
+// Copyright Epic Games, Inc. All Rights Reserved.
+
+#pragma once
+
+#include "CoreMinimal.h"
+#include "IAnalyticsSessionSummarySender.h"
+
+class IAnalyticsProviderET;
+
+/**
+ * Sends the analytics session summary to Epic Games analytics service.
+ */
+class ANALYTICSET_API FAnalyticsSessionSummarySender : public IAnalyticsSessionSummarySender
+{
+public:
+	/**
+	 * Construction a summary sender.
+	 * @param Provider The analytics provider that will be used to emits the summary event.
+	 * @param ShouldEmitFilterFunc A filter function invoked for each properties before sending. If the function is bound and returns false, the property will be filtered out. If unbound, all properties passed to SendSesionSummary() are emitted.
+	 */
+	FAnalyticsSessionSummarySender(IAnalyticsProviderET& Provider, TFunction<bool(const FString&, const FString&)> ShouldEmitFilterFunc = TFunction<bool(const FString&, const FString&)>());
+
+	/**
+	 * Emits the summary events for the specified session id on behalf of the specified user/app/appversion. The function filters the properties by invoking the functor specified at construction.
+	 * @param UserId The user emitting the report.
+	 * @param AppId The application for which the reports is emitted.
+	 * @param AppVersion The application version for which the report is emitted.
+	 * @param SessionId The session is for which the report is emitted.
+	 * @param Properties The list of properties that makes up the summary event.
+	 */
+	virtual bool SendSessionSummary(const FString& UserId, const FString& AppId, const FString& AppVersion, const FString& SessionId, const TMap<FString, FString>& Properties) override;
+
+private:
+	IAnalyticsProviderET& AnalyticsProvider;
+	TFunction<bool(const FString&, const FString&)> ShouldEmitPropFunc;
+};
diff --git a/Engine/Source/Runtime/Analytics/AnalyticsET/Public/IAnalyticsSessionSummarySender.h b/Engine/Source/Runtime/Analytics/AnalyticsET/Public/IAnalyticsSessionSummarySender.h
new file mode 100644
--- /dev/null
+++ b/Engine/Source/Runtime/Analytics/AnalyticsET/Public/IAnalyticsSessionSummarySender.h
@@ -0,0 +1,24 @@
+// Copyright Epic Games, Inc. All Rights Reserved.
+
+#pragma once
+
+#include "CoreMinimal.h"
+
+/**
+ * Sends the analytics session summary to a backend service.
+ */
+class IAnalyticsSessionSummarySender
+{
+public:
+	virtual ~IAnalyticsSessionSummarySender() = default;
+
+	/**
+	 * Emits the summary events for the specified session id on behalf of the specified user/app/appversion. The function filters the properties by invoking the functor specified at construction.
+	 * @param UserId The user emitting the report.
+	 * @param AppId The application for which the reports is emitted.
+	 * @param AppVersion The application version for which the report is emitted.
+	 * @param SessionId The session is for which the report is emitted.
+	 * @param Properties The list of properties that makes up the summary event.
+	 */
+	virtual bool SendSessionSummary(const FString& UserId, const FString& AppId, const FString& AppVersion, const FString& SessionId, const TMap<FString, FString>& Properties) = 0;
+};
diff --git a/Engine/Source/Runtime/Core/Private/Windows/WindowsPlatformCrashContext.cpp b/Engine/Source/Runtime/Core/Private/Windows/WindowsPlatformCrashContext.cpp
--- a/Engine/Source/Runtime/Core/Private/Windows/WindowsPlatformCrashContext.cpp
+++ b/Engine/Source/Runtime/Core/Private/Windows/WindowsPlatformCrashContext.cpp
@@ -480,6 +480,14 @@ FProcHandle LaunchCrashReportClient(void** OutWritePipe, void** OutReadPipe, uin
 		FCString::Strncat(CrashReporterClientArgs, PidStr, CR_CLIENT_MAX_ARGS_LEN);
 	}
 
+	{
+		// When CRC is spawned to 'monitor' a process, it creates an analytics session summary that can be merged with the summary of the watched process.
+		// The summaries are matched together using this process group ID.
+		TCHAR AnalyticsSummaryGuid[128] = { 0 };
+		FCString::Sprintf(AnalyticsSummaryGuid, TEXT(" -ProcessGroupId=%s"), *FApp::GetInstanceId().ToString(EGuidFormats::Digits));
+		FCString::Strncat(CrashReporterClientArgs, AnalyticsSummaryGuid, CR_CLIENT_MAX_ARGS_LEN);
+	}
+
 	// Parse commandline arguments relevant to pass to the client. Note that since we run this from static initialization
 	// FCommandline has not yet been initialized, instead we need to use the OS provided methods.
 	{
diff --git a/Engine/Source/Runtime/Engine/Private/Analytics/EditorAnalyticsSessionSummary.cpp b/Engine/Source/Runtime/Engine/Private/Analytics/EditorAnalyticsSessionSummary.cpp
new file mode 100644
--- /dev/null
+++ b/Engine/Source/Runtime/Engine/Private/Analytics/EditorAnalyticsSessionSummary.cpp
@@ -0,0 +1,194 @@
+// Copyright Epic Games, Inc. All Rights Reserved.
+
+#include "EditorAnalyticsSessionSummary.h"
+
+#if WITH_EDITOR
+
+#include "Editor.h"
+#include "IVREditorModule.h"
+#include "Framework/Application/SlateApplication.h"
+#include "HAL/PlatformTime.h"
+#include "Interfaces/IAnalyticsPropertyStore.h"
+#include "Kismet2/DebuggerCommands.h"
+#include "ProfilingDebugging/StallDetector.h"
+
+namespace EditorAnalyticsProperties
+{
+	// List of mutable keys.
+	static const TAnalyticsProperty<uint32>  UserInteractionCount   = TEXT("UserInteractionCount");
+	static const TAnalyticsProperty<int32>   Idle1Min               = TEXT("1MinIdle");
+	static const TAnalyticsProperty<int32>   Idle5Min               = TEXT("5MinIdle");
+	static const TAnalyticsProperty<int32>   Idle30Min              = TEXT("30MinIdle");
+	static const TAnalyticsProperty<bool>    IsInPIE                = TEXT("IsInPIE");
+	static const TAnalyticsProperty<bool>    IsInVRMode             = TEXT("IsInVRMode");
+
+	static const TAnalyticsProperty<uint32>  TotalStallCount        = TEXT("TotalStallCount");
+	static const TAnalyticsProperty<uint32>  TotalStallReported     = TEXT("TotalStallReported");
+	static const TAnalyticsProperty<uint32>  TopStallTriggerCount   = TEXT("TopStallTriggerCount");
+	static const TAnalyticsProperty<FString> TopStallName           = TEXT("TopStallName");
+	static const TAnalyticsProperty<double>  TopStallBudgetSeconds  = TEXT("TopStallBudgetSeconds");
+	static const TAnalyticsProperty<double>  TopStallOverageSeconds = TEXT("TopStallOverageSeconds");
+	static const TAnalyticsProperty<uint32>  ProcessDiagnostics     = TEXT("ProcessDiagnostics"); // Whether some profiling/diagnostic tools are enabled, which could slow down the Editor.
+	static const TAnalyticsProperty<uint32>  SummaryEventVersion    = TEXT("SummaryEventVersion"); // A version number used identify the key/set used. Can be used to compare before/after some changes too since comparing engine version is not always straingth foward.
+}
+
+FEditorAnalyticsSessionSummary::FEditorAnalyticsSessionSummary(TSharedPtr<IAnalyticsPropertyStore> Store, uint32 MonitorProcessId)
+	: FEngineAnalyticsSessionSummary(MoveTemp(Store), MonitorProcessId)
+	, LastUserActivityTimeSecs(FPlatformTime::Seconds())
+	, AccountedUserIdleSecs(0)
+{
+	EditorAnalyticsProperties::UserInteractionCount.Set(GetStore(), 0);
+	EditorAnalyticsProperties::Idle1Min.Set(GetStore(), 0);
+	EditorAnalyticsProperties::Idle5Min.Set(GetStore(), 0);
+	EditorAnalyticsProperties::Idle30Min.Set(GetStore(), 0);
+	EditorAnalyticsProperties::IsInVRMode.Set(GetStore(), IVREditorModule::Get().IsVREditorModeActive());
+	EditorAnalyticsProperties::IsInPIE.Set(GetStore(), FPlayWorldCommandCallbacks::IsInPIE());
+
+	EditorAnalyticsProperties::TopStallName.Set(GetStore(), TEXT(""), /*Capacity*/128);
+	EditorAnalyticsProperties::TopStallBudgetSeconds.Set(GetStore(), 0.0);
+	EditorAnalyticsProperties::TopStallOverageSeconds.Set(GetStore(), 0.0);
+	EditorAnalyticsProperties::TopStallTriggerCount.Set(GetStore(), 0);
+	EditorAnalyticsProperties::TotalStallCount.Set(GetStore(), 0);
+	EditorAnalyticsProperties::TotalStallReported.Set(GetStore(), 0);
+	EditorAnalyticsProperties::ProcessDiagnostics.Set(GetStore(), static_cast<uint32>(FPlatformMisc::GetProcessDiagnostics()));
+
+	// The current summary revision number. Identify the key set used and its behavior. If we add/remove keys, we should increment the number. The number can
+	// also be incremented when we change some behaviors to be able to compare between versions.
+	EditorAnalyticsProperties::SummaryEventVersion.Set(GetStore(), 1);
+
+	// Persist the session to disk.
+	GetStore()->Flush();
+
+	FEditorDelegates::PreBeginPIE.AddRaw(this, &FEditorAnalyticsSessionSummary::OnEnterPIE);
+	FEditorDelegates::EndPIE.AddRaw(this, &FEditorAnalyticsSessionSummary::OnExitPIE);
+	FSlateApplication::Get().GetOnModalLoopTickEvent().AddRaw(this, &FEditorAnalyticsSessionSummary::Tick);
+	FSlateApplication::Get().GetLastUserInteractionTimeUpdateEvent().AddRaw(this, &FEditorAnalyticsSessionSummary::OnSlateUserInteraction);
+}
+
+void FEditorAnalyticsSessionSummary::ShutdownInternal()
+{
+	FEditorDelegates::PreBeginPIE.RemoveAll(this);
+	FEditorDelegates::EndPIE.RemoveAll(this);
+	FSlateApplication::Get().GetOnModalLoopTickEvent().RemoveAll(this);
+	FSlateApplication::Get().GetLastUserInteractionTimeUpdateEvent().RemoveAll(this);
+}
+
+bool FEditorAnalyticsSessionSummary::UpdateSessionProgressInternal(bool bCrashing)
+{
+	bool bShouldPersist = UpdateUserIdleTime(FPlatformTime::Seconds(), /*bReset*/false);
+
+	// In case of crash, don't save anything else.
+	if (bCrashing)
+	{
+		return true; // Yes, should persist.
+	}
+
+	EditorAnalyticsProperties::IsInVRMode.Set(GetStore(), IVREditorModule::Get().IsVREditorModeActive());
+	EditorAnalyticsProperties::IsInPIE.Set(GetStore(), FPlayWorldCommandCallbacks::IsInPIE());
+
+	// Accumulate stall stats
+#if STALL_DETECTOR
+	TArray<UE::FStallDetectorStats::TabulatedResult> StallResults;
+	UE::FStallDetectorStats::TabulateStats(StallResults);
+	if (!StallResults.IsEmpty())
+	{
+		UE::FStallDetectorStats::TabulatedResult TopResult(StallResults[0]);
+		EditorAnalyticsProperties::TopStallName.Set(GetStore(), TopResult.Stats->Name, /*Capacity*/128);
+		EditorAnalyticsProperties::TopStallBudgetSeconds.Set(GetStore(), TopResult.Stats->BudgetSeconds);
+		EditorAnalyticsProperties::TopStallOverageSeconds.Set(GetStore(), TopResult.OverageSeconds);
+		EditorAnalyticsProperties::TopStallTriggerCount.Set(GetStore(), TopResult.TriggerCount);
+	}
+
+	EditorAnalyticsProperties::TotalStallCount.Set(GetStore(), UE::FStallDetectorStats::TotalTriggeredCount.Get());
+	EditorAnalyticsProperties::TotalStallReported.Set(GetStore(), UE::FStallDetectorStats::TotalReportedCount.Get());
+#endif // STALL_DETECTOR
+
+	return bShouldPersist;
+}
+
+bool FEditorAnalyticsSessionSummary::UpdateUserIdleTime(double CurrTimeSecs, bool bReset)
+{
+	bool bSessionUpdated = false;
+
+	// How much time elapsed since the last activity.
+	double TotalIdleSecs = CurrTimeSecs - LastUserActivityTimeSecs.load();
+	if (TotalIdleSecs > 60.0) // Less than a minute is always considered normal interaction delay.
+	{
+		double LastAccountedIdleSecs = AccountedUserIdleSecs.load();
+		double UnaccountedIdleSecs   = TotalIdleSecs - LastAccountedIdleSecs;
+
+		// If one or more minute is unaccounted
+		if (UnaccountedIdleSecs >= 60.0)
+		{
+			double AccountedIdleMins = FMath::FloorToDouble(LastAccountedIdleSecs/60); // Minutes already accounted for.
+			double ToAccountIdleMins = FMath::FloorToDouble(UnaccountedIdleSecs/60);   // New minutes to account for (entire minute only)
+
+			// Delta = LatestAccounted - AlreadyAccounted
+			double DeltaIdle1Min  = FMath::Max(0.0, AccountedIdleMins + ToAccountIdleMins - 1)  - FMath::Max(0.0, AccountedIdleMins - 1);  // The first minute of this idle sequence is considered 'normal interaction delay' and in not accounted as idle.
+			double DeltaIdle5Min  = FMath::Max(0.0, AccountedIdleMins + ToAccountIdleMins - 5)  - FMath::Max(0.0, AccountedIdleMins - 5);  // The 5 first minutes of this idle sequence are considered 'normal interaction delay' and are not accounted for the 5-min timer.
+			double DeltaIdle30Min = FMath::Max(0.0, AccountedIdleMins + ToAccountIdleMins - 30) - FMath::Max(0.0, AccountedIdleMins - 30); // The 30 first minutes of this idle sequence are considered 'normal interaction delay' and are not accounted for the 30-min timer.
+
+			// Ensure only one thread adds the current delta time.
+			if (AccountedUserIdleSecs.compare_exchange_strong(LastAccountedIdleSecs, LastAccountedIdleSecs + ToAccountIdleMins * 60.0)) // Only add the 'accounted' minutes and keep fraction of minutes running.
+			{
+				EditorAnalyticsProperties::Idle1Min.Update(GetStore(), [DeltaIdle1Min](int32& InOutValue)
+				{
+					InOutValue += FMath::RoundToInt(static_cast<float>(DeltaIdle1Min));
+					return true;
+				});
+
+				EditorAnalyticsProperties::Idle5Min.Update(GetStore(), [DeltaIdle5Min](int32& InOutValue)
+				{
+					InOutValue += FMath::RoundToInt(static_cast<float>(DeltaIdle5Min));
+					return true;
+				});
+
+				EditorAnalyticsProperties::Idle30Min.Update(GetStore(), [DeltaIdle30Min](int32& InOutValue)
+				{
+					InOutValue += FMath::RoundToInt(static_cast<float>(DeltaIdle30Min));
+					return true;
+				});
+
+				bSessionUpdated = true;
+			}
+		}
+	}
+
+	if (bReset)
+	{
+		AccountedUserIdleSecs.store(0);
+		LastUserActivityTimeSecs.store(CurrTimeSecs);
+	}
+
+	return bSessionUpdated; // True if the idle timers were updated.
+}
+
+void FEditorAnalyticsSessionSummary::OnSlateUserInteraction(double CurrSlateInteractionTime)
+{
+	EditorAnalyticsProperties::UserInteractionCount.Update(GetStore(), [](uint32& InOutValue)
+	{
+		++InOutValue;
+		return true;
+	});
+
+	// If the user input 'reset' the idle timers.
+	double CurrTimeSecs = FPlatformTime::Seconds();
+	if (UpdateUserIdleTime(CurrTimeSecs, /*bReset*/true))
+	{
+		GetStore()->Flush(/*bAsync*/true, FTimespan::Zero());
+	}
+}
+
+void FEditorAnalyticsSessionSummary::OnEnterPIE(const bool /*bIsSimulating*/)
+{
+	EditorAnalyticsProperties::IsInPIE.Set(GetStore(), true);
+	GetStore()->Flush(/*bAsync*/true, FTimespan::Zero());
+}
+
+void FEditorAnalyticsSessionSummary::OnExitPIE(const bool /*bIsSimulating*/)
+{
+	EditorAnalyticsProperties::IsInPIE.Set(GetStore(), false);
+	GetStore()->Flush(/*bAsync*/true, FTimespan::Zero());
+}
+
+#endif // WITH_EDITOR
diff --git a/Engine/Source/Runtime/Engine/Private/Analytics/EditorAnalyticsSessionSummary.h b/Engine/Source/Runtime/Engine/Private/Analytics/EditorAnalyticsSessionSummary.h
new file mode 100644
--- /dev/null
+++ b/Engine/Source/Runtime/Engine/Private/Analytics/EditorAnalyticsSessionSummary.h
@@ -0,0 +1,39 @@
+// Copyright Epic Games, Inc. All Rights Reserved.
+
+#pragma once
+
+#include "CoreMinimal.h"
+
+#if WITH_EDITOR
+
+#include "Analytics/EngineAnalyticsSessionSummary.h"
+#include <atomic>
+
+/** Specializes the base engine analytics for the Editor, adding Editor only data to the summary. */
+class FEditorAnalyticsSessionSummary : public FEngineAnalyticsSessionSummary
+{
+public:
+	FEditorAnalyticsSessionSummary(TSharedPtr<IAnalyticsPropertyStore> Store, uint32 MonitorProcessId);
+
+protected:
+	/** Invoked by base class FEngineAnalyticsSessionSummary at the configured update period. */
+	virtual bool UpdateSessionProgressInternal(bool bCrashing) override;
+
+	/** Invoked by base class FEngineAnalyticsSessionSummary on shutdown. */
+	virtual void ShutdownInternal() override;
+
+private:
+	bool UpdateUserIdleTime(double CurrTimeSecs, bool bReset);
+	void OnSlateUserInteraction(double CurrSlateInteractionTime);
+	void OnEnterPIE(const bool /*bIsSimulating*/);
+	void OnExitPIE(const bool /*bIsSimulating*/);
+
+private:
+	/** Last activity (user input, crash, terminate, shutdown) timestamp from FPlatformTime::Seconds() to track user inactivity. */
+	std::atomic<double> LastUserActivityTimeSecs;
+
+	/** The number of idle seconds in the current idle sequence that were accounted (saved in the session) for the user idle counters. */
+	std::atomic<double> AccountedUserIdleSecs;
+};
+
+#endif // WITH_EDITOR
diff --git a/Engine/Source/Runtime/Engine/Private/Analytics/EngineAnalyticsSessionSummary.cpp b/Engine/Source/Runtime/Engine/Private/Analytics/EngineAnalyticsSessionSummary.cpp
new file mode 100644
--- /dev/null
+++ b/Engine/Source/Runtime/Engine/Private/Analytics/EngineAnalyticsSessionSummary.cpp
@@ -0,0 +1,384 @@
+// Copyright Epic Games, Inc. All Rights Reserved.
+
+#include "Analytics/EngineAnalyticsSessionSummary.h"
+#include "AnalyticsSessionSummaryManager.h"
+#include "EngineGlobals.h"
+#include "GeneralProjectSettings.h"
+#include "GenericPlatform/GenericPlatformCrashContext.h"
+#include "HAL/ExceptionHandling.h"
+#include "HAL/FileManager.h"
+#include "HAL/PlatformProcess.h"
+#include "HAL/PlatformTime.h"
+#include "Interfaces/IAnalyticsPropertyStore.h"
+#include "Interfaces/IPluginManager.h"
+#include "Interfaces/IProjectManager.h"
+#include "Misc/EngineVersion.h"
+#include "UserActivityTracking.h"
+
+extern ENGINE_API float GAverageFPS;
+extern CORE_API bool GIsGPUCrashed;
+
+namespace EngineAnalyticsProperties
+{
+	// List of mutable properties.
+	static const TAnalyticsProperty<FDateTime> Timestamp         = TEXT("Timestamp");
+	static const TAnalyticsProperty<FDateTime> LastTickTimestamp = TEXT("LastTickTimestamp");
+	static const TAnalyticsProperty<int32> SessionDurationSecs   = TEXT("SessionDuration");
+	static const TAnalyticsProperty<uint64> EngineTickCount      = TEXT("EngineTickCount");
+	static const TAnalyticsProperty<uint64> SessionTickCount     = TEXT("SessionTickCount");
+	static const TAnalyticsProperty<FString> UserActivity        = TEXT("CurrentUserActivity");
+	static const TAnalyticsProperty<int32>  MonitorExitCode      = TEXT("MonitorExitCode");
+	static const TAnalyticsProperty<float> AverageFPS            = TEXT("AverageFPS");
+	static const TAnalyticsProperty<bool> IsTerminating          = TEXT("IsTerminating");
+	static const TAnalyticsProperty<bool> WasShutdown            = TEXT("WasShutdown");
+	static const TAnalyticsProperty<bool> IsCrashing             = TEXT("IsCrash");
+	static const TAnalyticsProperty<bool> IsGpuCrash             = TEXT("GPUCrash");
+	static const TAnalyticsProperty<bool> IsVanilla              = TEXT("IsVanilla");
+	static const TAnalyticsProperty<bool> IsLowDriveSpace        = TEXT("IsLowDriveSpace");
+	static const TAnalyticsProperty<bool> WasDebuggerIgnored     = TEXT("WasDebuggerIgnored");
+	static const TAnalyticsProperty<bool> WasEverDebugged        = TEXT("WasDebugged");
+	static const TAnalyticsProperty<bool> IsCrcMissing           = TEXT("IsCrcMissing");
+	static const TAnalyticsProperty<bool> RunningOnBattery       = TEXT("RunningOnBattery");
+	static const TAnalyticsProperty<bool> IsUserLoggingOut       = FAnalyticsSessionSummaryManager::IsUserLoggingOutProperty;
+	static const TAnalyticsProperty<int32> ShutdownType          = FAnalyticsSessionSummaryManager::ShutdownTypeCodeProperty;
+
+	static FString UnknownUserActivity(TEXT("Unknown"));
+
+	float GetAverageFPS()
+	{
+		return GAverageFPS;
+	}
+
+	uint64 GetEngineTickCount()
+	{
+		return GFrameCounter;
+	}
+
+	bool IsCrcExecutableMissing()
+	{
+		// TODO: Add a function later on to PlatfomCrashContext to check existance of CRC. This is only used on windows at the moment to categorize cases where CRC fails to report the exit code (only supported on Windows).
+	#if PLATFORM_WINDOWS
+		const FString EngineDir = FPlatformMisc::EngineDir();
+
+		// Find the path to crash reporter binary. Avoid creating FStrings.
+		FString CrcPathRel = EngineDir / TEXT("Binaries") / FPlatformProcess::GetBinariesSubdirectory() / TEXT("CrashReportClientEditor.exe");
+		FString CrcPathDev = EngineDir / TEXT("Binaries") / FPlatformProcess::GetBinariesSubdirectory() / TEXT("CrashReportClientEditor-Win64-Development.exe");
+		FString CrcPathDbg = EngineDir / TEXT("Binaries") / FPlatformProcess::GetBinariesSubdirectory() / TEXT("CrashReportClientEditor-Win64-Debug.exe");
+
+		return !IFileManager::Get().FileExists(*CrcPathRel) && !IFileManager::Get().FileExists(*CrcPathDev) && !IFileManager::Get().FileExists(*CrcPathDbg);
+	#else
+		return false;
+	#endif
+	}
+
+	FString GetProjectName(const UGeneralProjectSettings& ProjectSettings)
+	{
+		FString ProjectName = FApp::GetProjectName();
+		if (ProjectName.Len() && ProjectSettings.ProjectName.Len())
+		{
+			if (ProjectSettings.ProjectName != ProjectName)
+			{
+				ProjectName = ProjectName / ProjectSettings.ProjectName; // The project names don't match, report both.
+			}
+		}
+		else if (ProjectName.IsEmpty())
+		{
+			ProjectName = ProjectSettings.ProjectName;
+		}
+
+		return ProjectName;
+	}
+
+	FString GetUserActivityString(const FUserActivity& InUserActivity)
+	{
+		return InUserActivity.ActionName.IsEmpty() ? UnknownUserActivity : InUserActivity.ActionName;
+	}
+
+	bool ShouldOverwriteShutdownType(const int32* ActualShutdownType, const int32& ProposedShutdownType)
+	{
+		if (static_cast<EAnalyticsSessionShutdownType>(ProposedShutdownType) == EAnalyticsSessionShutdownType::Crashed) // Crash overwrite everything else.
+		{
+			return true;
+		}
+		else if (static_cast<EAnalyticsSessionShutdownType>(*ActualShutdownType) == EAnalyticsSessionShutdownType::Crashed) // Cannot overwrite a crash.
+		{
+			return false;
+		}
+		else if (static_cast<EAnalyticsSessionShutdownType>(ProposedShutdownType) == EAnalyticsSessionShutdownType::Debugged) // Debugger can overwrite Shutdown|Terminated|Abnormal
+		{
+			return true;
+		}
+		else if (static_cast<EAnalyticsSessionShutdownType>(*ActualShutdownType) == EAnalyticsSessionShutdownType::Debugged) // Shutdown|Terminated|Abnormal cannot overwrite the debugger.
+		{
+			return false;
+		}
+		return true; // Shutdown|Terminate|Abnormal can overwrite each other.
+	}
+
+} // namespace EngineAnalyticsProperties
+
+
+FEngineAnalyticsSessionSummary::FEngineAnalyticsSessionSummary(TSharedPtr<IAnalyticsPropertyStore> InStore, uint32 InMonitorProcessId)
+	: Store(MoveTemp(InStore))
+	, SessionStartTimeUtc(FDateTime::UtcNow())
+	, SessionStartTimeSecs(FPlatformTime::Seconds())
+	, CrcProcessId(InMonitorProcessId)
+	, PersistPeriod(FTimespan::FromSeconds(10))
+{
+	// List the plugins.
+	TArray<TSharedRef<IPlugin>> Plugins = IPluginManager::Get().GetEnabledPlugins();
+	TArray<FString> PluginNames;
+	PluginNames.Reserve(Plugins.Num());
+	for (const TSharedRef<IPlugin>& Plugin : Plugins)
+	{
+		PluginNames.Add(Plugin->GetName());
+	}
+	PluginNames.Sort();
+
+	// Get the operating system information.
+	FString OSMajor;
+	FString OSMinor;
+	FPlatformMisc::GetOSVersions(/*out*/ OSMajor, /*out*/ OSMinor);
+
+	// Get project settings.
+	const UGeneralProjectSettings& ProjectSettings = *GetDefault<UGeneralProjectSettings>();
+
+	// Create the immutable metrics.
+	Store->Set(TEXT("Platform"), FString(FPlatformProperties::IniPlatformName()));
+	Store->Set(TEXT("OSMajor"), OSMajor);
+	Store->Set(TEXT("OSMinor"), OSMinor);
+	Store->Set(TEXT("OSVersion"), FPlatformMisc::GetOSVersion());
+	Store->Set(TEXT("Is64BitOS"), FPlatformMisc::Is64bitOperatingSystem());
+	Store->Set(TEXT("TotalPhysicalRAM"), FPlatformMemory::GetStats().TotalPhysical);
+	Store->Set(TEXT("CPUPhysicalCores"), FPlatformMisc::NumberOfCores());
+	Store->Set(TEXT("CPULogicalCores"), FPlatformMisc::NumberOfCoresIncludingHyperthreads());
+	Store->Set(TEXT("CPUVendor"), FPlatformMisc::GetCPUVendor());
+	Store->Set(TEXT("CPUBrand"), FPlatformMisc::GetCPUBrand());
+	Store->Set(TEXT("DesktopGPUAdapter"), FPlatformMisc::GetPrimaryGPUBrand());
+	Store->Set(TEXT("RenderingGPUAdapter"), GRHIAdapterName);
+	Store->Set(TEXT("GPUVendorID"), GRHIVendorId);
+	Store->Set(TEXT("GPUDeviceID"), GRHIDeviceId);
+	Store->Set(TEXT("GRHIDeviceRevision"), GRHIDeviceRevision);
+	Store->Set(TEXT("GRHIAdapterInternalDriverVersion"), GRHIAdapterInternalDriverVersion);
+	Store->Set(TEXT("GRHIAdapterUserDriverVersion"), GRHIAdapterUserDriverVersion);
+	Store->Set(TEXT("GRHIName"), FString(GDynamicRHI ? GDynamicRHI->GetName() : TEXT("")));
+	Store->Set(TEXT("ProjectName"), EngineAnalyticsProperties::GetProjectName(ProjectSettings));
+	Store->Set(TEXT("ProjectID"), ProjectSettings.ProjectID.ToString(EGuidFormats::DigitsWithHyphens));
+	Store->Set(TEXT("ProjectDescription"), ProjectSettings.Description);
+	Store->Set(TEXT("ProjectVersion"), ProjectSettings.ProjectVersion);
+	Store->Set(TEXT("EngineVersion"), FEngineVersion::Current().ToString(EVersionComponent::Changelist));
+	Store->Set(TEXT("CommandLine"), FString(FCommandLine::GetForLogging()));
+	Store->Set(TEXT("MonitorPid"), CrcProcessId); // CrashReportClient acts as the monitoring/reporting process.
+	Store->Set(TEXT("Plugins"), FString::Join(PluginNames, TEXT(",")));
+	Store->Set(TEXT("IsCrcMissing"), CrcProcessId != 0 ? false : EngineAnalyticsProperties::IsCrcExecutableMissing()); // CrashReportClient is the monitor process. If we have a process id, it's not missing.
+	Store->Set(TEXT("IsInEnterprise"), IProjectManager::Get().IsEnterpriseProject());
+	Store->Set(TEXT("StartupTimestamp"), SessionStartTimeUtc);
+
+	FDateTime CurrTimeUtc = FDateTime::UtcNow();
+
+	// Create the mutable metrics.
+	EngineAnalyticsProperties::Timestamp.Set(GetStore(), CurrTimeUtc);
+	EngineAnalyticsProperties::LastTickTimestamp.Set(GetStore(), CurrTimeUtc);
+	EngineAnalyticsProperties::SessionDurationSecs.Set(GetStore(), FPlatformTime::Seconds() - SessionStartTimeSecs);
+	EngineAnalyticsProperties::EngineTickCount.Set(GetStore(), EngineAnalyticsProperties::GetEngineTickCount());
+	EngineAnalyticsProperties::SessionTickCount.Set(GetStore(), CurrSessionTickCount);
+	EngineAnalyticsProperties::UserActivity.Set(GetStore(), EngineAnalyticsProperties::GetUserActivityString(FUserActivityTracking::GetUserActivity()));
+	EngineAnalyticsProperties::AverageFPS.Set(GetStore(), EngineAnalyticsProperties::GetAverageFPS());
+	EngineAnalyticsProperties::IsTerminating.Set(GetStore(),  false);
+	EngineAnalyticsProperties::WasShutdown.Set(GetStore(), false);
+	EngineAnalyticsProperties::IsCrashing.Set(GetStore(), false);
+	EngineAnalyticsProperties::IsGpuCrash.Set(GetStore(), false);
+	EngineAnalyticsProperties::IsUserLoggingOut.Set(GetStore(), false);
+	EngineAnalyticsProperties::IsVanilla.Set(GetStore(), GEngine && GEngine->IsVanillaProduct());
+	EngineAnalyticsProperties::IsLowDriveSpace.Set(GetStore(), false);
+	EngineAnalyticsProperties::WasDebuggerIgnored.Set(GetStore(), false);
+	EngineAnalyticsProperties::WasEverDebugged.Set(GetStore(), false);
+	EngineAnalyticsProperties::ShutdownType.Set(GetStore(), (int32)EAnalyticsSessionShutdownType::Abnormal);
+	EngineAnalyticsProperties::RunningOnBattery.Set(GetStore(), FPlatformMisc::IsRunningOnBattery());
+
+	Store->Flush();
+
+	// Listen to interesting events.
+	FCoreDelegates::OnHandleSystemError.AddRaw(this, &FEngineAnalyticsSessionSummary::OnCrashing); // WARNING: Don't assume this function is only called from game thread.
+	FCoreDelegates::ApplicationWillTerminateDelegate.AddRaw(this, &FEngineAnalyticsSessionSummary::OnTerminate); // WARNING: Don't assume this function is only called from game thread.
+	FCoreDelegates::IsVanillaProductChanged.AddRaw(this, &FEngineAnalyticsSessionSummary::OnVanillaStateChanged);
+	FCoreDelegates::OnUserLoginChangedEvent.AddRaw(this, &FEngineAnalyticsSessionSummary::OnUserLoginChanged);
+	FUserActivityTracking::OnActivityChanged.AddRaw(this, &FEngineAnalyticsSessionSummary::OnUserActivity);
+}
+
+void FEngineAnalyticsSessionSummary::Shutdown()
+{
+	if (!bShutdown)
+	{
+		// Call the derived class.
+		ShutdownInternal();
+
+		FCoreDelegates::OnHandleSystemError.RemoveAll(this);
+		FCoreDelegates::ApplicationWillTerminateDelegate.RemoveAll(this);
+		FCoreDelegates::IsVanillaProductChanged.RemoveAll(this);
+		FCoreDelegates::OnUserLoginChangedEvent.RemoveAll(this);
+		FUserActivityTracking::OnActivityChanged.RemoveAll(this);
+
+		EngineAnalyticsProperties::WasShutdown.Set(GetStore(), true);
+		EngineAnalyticsProperties::ShutdownType.Set(GetStore(), (int32)EAnalyticsSessionShutdownType::Shutdown, EngineAnalyticsProperties::ShouldOverwriteShutdownType);
+		UpdateSessionProgress(/*bFlushAsync*/false, /*FlushTimeout*/FTimespan::MaxValue()); // Forces a sunchronous flush.
+		bShutdown = true;
+	}
+}
+
+void FEngineAnalyticsSessionSummary::Tick(float DeltaTime)
+{
+	if (bShutdown)
+	{
+		return;
+	}
+
+	// The number of times this instance was ticked.
+	EngineAnalyticsProperties::SessionTickCount.Update(GetStore(), [](uint64& Value) { ++Value; return true; });
+
+	// The average FPS.
+	EngineAnalyticsProperties::AverageFPS.Set(GetStore(), EngineAnalyticsProperties::GetAverageFPS());
+
+	// This tick timestamp. Not allowed to go back.
+	EngineAnalyticsProperties::LastTickTimestamp.Set(GetStore(), FDateTime::UtcNow(), [](const FDateTime* Actual, const FDateTime& Proposed){ return Proposed >*Actual; });
+
+	// Detect if CRC state changed since the last update.
+	bool bFlush = UpdateExternalProcessReporterState(/*bQuickCheck*/true);
+
+	// Check if the debugger states changed.
+	bFlush |= UpdateDebuggerStates();
+
+	// Ensure to flush the properties periodically.
+	if (bFlush || FPlatformTime::Seconds() >= NextPersistTimeSeconds)
+	{
+		UpdateSessionProgress();
+		NextPersistTimeSeconds = FPlatformTime::Seconds() + PersistPeriod.GetTotalSeconds();
+	}
+}
+
+void FEngineAnalyticsSessionSummary::OnCrashing()
+{
+	EngineAnalyticsProperties::IsCrashing.Set(GetStore(), true);
+	EngineAnalyticsProperties::IsGpuCrash.Set(GetStore(), GIsGPUCrashed);
+	EngineAnalyticsProperties::ShutdownType.Set(GetStore(), (int32)EAnalyticsSessionShutdownType::Crashed);
+
+	// Flush the store if possible but don't wait forever to avoid deadlocking if the crash affects the store ability to complete its background task.
+	UpdateSessionProgress(/*bFlushAsync*/false, /*FlushTimeout*/FTimespan::FromMilliseconds(100), /*bCrashing*/true);
+}
+
+void FEngineAnalyticsSessionSummary::OnTerminate()
+{
+	EngineAnalyticsProperties::IsTerminating.Set(GetStore(), true);
+	EngineAnalyticsProperties::ShutdownType.Set(GetStore(), (int32)EAnalyticsSessionShutdownType::Terminated, EngineAnalyticsProperties::ShouldOverwriteShutdownType);
+	UpdateSessionProgress(/*bFlushAsync*/false, /*FlushTimeout*/FTimespan::FromMilliseconds(100));
+}
+
+void FEngineAnalyticsSessionSummary::OnUserLoginChanged(bool bLoggingIn, int32, int32)
+{
+	if (!bLoggingIn)
+	{
+		EngineAnalyticsProperties::IsUserLoggingOut.Set(GetStore(), true);
+		UpdateSessionProgress(/*bFlushAsync*/false, /*FlushTimeout*/FTimespan::FromMilliseconds(100));
+	}
+}
+
+void FEngineAnalyticsSessionSummary::OnVanillaStateChanged(bool bIsVanilla)
+{
+	EngineAnalyticsProperties::IsVanilla.Set(GetStore(), bIsVanilla);
+	UpdateSessionProgress();
+}
+
+void FEngineAnalyticsSessionSummary::LowDriveSpaceDetected()
+{
+	EngineAnalyticsProperties::IsLowDriveSpace.Set(GetStore(), true);
+	UpdateSessionProgress();
+}
+
+void FEngineAnalyticsSessionSummary::OnUserActivity(const FUserActivity& InUserActivity)
+{
+	EngineAnalyticsProperties::UserActivity.Set(GetStore(), EngineAnalyticsProperties::GetUserActivityString(InUserActivity));
+	UpdateSessionProgress();
+}
+
+void FEngineAnalyticsSessionSummary::UpdateSessionProgress(bool bFlushAsync, const FTimespan& FlushTimeout, bool bCrashing)
+{
+	// Let the derived classes update.
+	UpdateSessionProgressInternal(bCrashing);
+
+	// Current session duration in seconds.
+	EngineAnalyticsProperties::SessionDurationSecs.Set(GetStore(), FMath::FloorToInt(static_cast<float>(FPlatformTime::Seconds() - SessionStartTimeSecs)),
+		[](const int32* Actual, const int32& Proposed) { return Proposed > *Actual; });
+
+	// Wall timestamp.
+	EngineAnalyticsProperties::Timestamp.Set(GetStore(), FDateTime::UtcNow(),
+		[](const FDateTime* Actual, const FDateTime& Proposed) { return Proposed > *Actual; });
+
+	// The number of times the engine ticked.
+	EngineAnalyticsProperties::EngineTickCount.Set(GetStore(), EngineAnalyticsProperties::GetEngineTickCount(),
+		[](const uint64* Actual, const uint64& Proposed) { return Proposed > *Actual; });
+
+	// Flush (or try to flush) the store.
+	Store->Flush(bFlushAsync, FlushTimeout);
+}
+
+bool FEngineAnalyticsSessionSummary::UpdateExternalProcessReporterState(bool bQuickCheck)
+{
+	if (CrcProcessId == 0)
+	{
+		return false; // Nothing to update, monitor is not running in background (not supported/not in monitor mode/failed to launch)
+	}
+	else if (CrcExitCode.IsSet() && *CrcExitCode != ECrashExitCodes::OutOfProcessReporterExitedUnexpectedly)
+	{
+		return false; // Already have the real exit code set.
+	}
+	else if (TOptional<int32> ExitCode = FGenericCrashContext::GetOutOfProcessCrashReporterExitCode())
+	{
+		CrcExitCode = MoveTemp(ExitCode);
+		EngineAnalyticsProperties::MonitorExitCode.Set(GetStore(), *CrcExitCode);
+		UpdateSessionProgress();
+		return true;
+	}
+	else if (bQuickCheck)
+	{
+		return false; // All the code above is pretty fast and can run every tick. IsApplicationRunning() is slow, so exit here.
+	}
+	else if (!CrcExitCode.IsSet() && !FPlatformProcess::IsApplicationRunning(CrcProcessId))
+	{
+		// Set a rather unique, but known exit code as place holder, hoping that next update, the engine will report the real one.
+		CrcExitCode.Emplace(ECrashExitCodes::OutOfProcessReporterExitedUnexpectedly);
+		EngineAnalyticsProperties::MonitorExitCode.Set(GetStore(), *CrcExitCode);
+		UpdateSessionProgress();
+		return true;
+	}
+	// else -> either CrashReportClientEditor is still running or we already flagged it as dead.
+	return false;
+
+	return false;
+}
+
+bool FEngineAnalyticsSessionSummary::UpdateDebuggerStates()
+{
+	bool bStateChanged = false;
+
+	// Ignoring the debugger changes how IsDebuggerPresent() behave and masks the usage of the debugger if true.
+	if (!bDebuggerIgnored && GIgnoreDebugger)
+	{
+		bStateChanged = true;
+		bDebuggerIgnored = true;
+		bWasEverDebugged = true;
+		EngineAnalyticsProperties::WasDebuggerIgnored.Set(GetStore(), true);
+		EngineAnalyticsProperties::WasEverDebugged.Set(GetStore(), true); // Assumes the session was debugged if the debugger was ignored.
+		EngineAnalyticsProperties::ShutdownType.Set(GetStore(), (int32)EAnalyticsSessionShutdownType::Debugged, EngineAnalyticsProperties::ShouldOverwriteShutdownType);
+	}
+
+	// Check if the debugger is attached.
+	if (!bWasEverDebugged && FPlatformMisc::IsDebuggerPresent())
+	{
+		bStateChanged = true;
+		bWasEverDebugged = true;
+		EngineAnalyticsProperties::WasEverDebugged.Set(GetStore(), true);
+		EngineAnalyticsProperties::ShutdownType.Set(GetStore(), (int32)EAnalyticsSessionShutdownType::Debugged, EngineAnalyticsProperties::ShouldOverwriteShutdownType);
+	}
+
+	return bStateChanged;
+}
diff --git a/Engine/Source/Runtime/Engine/Private/EngineAnalytics.cpp b/Engine/Source/Runtime/Engine/Private/EngineAnalytics.cpp
--- a/Engine/Source/Runtime/Engine/Private/EngineAnalytics.cpp
+++ b/Engine/Source/Runtime/Engine/Private/EngineAnalytics.cpp
@@ -1,6 +1,7 @@
 // Copyright Epic Games, Inc. All Rights Reserved.
 
 #include "EngineAnalytics.h"
+#include "Misc/App.h"
 #include "Misc/Guid.h"
 #include "Stats/Stats.h"
 #include "Misc/ConfigCacheIni.h"
@@ -19,9 +20,11 @@
 #include "StudioAnalytics.h"
 
 #if WITH_EDITOR
-#include "EditorAnalyticsSession.h"
-#include "EditorSessionSummarySender.h"
-#include "Analytics/EditorSessionSummaryWriter.h"
+#include "AnalyticsPropertyStore.h"
+#include "AnalyticsSessionSummaryManager.h"
+#include "AnalyticsSessionSummarySender.h"
+#include "Analytics/EditorAnalyticsSessionSummary.h"
+#include "EditorAnalyticsSession.h" // DEPRECATED: kept around to clean up expired old sessions.
 #endif
 
 bool FEngineAnalytics::bIsInitialized;
@@ -29,8 +32,9 @@ TSharedPtr<IAnalyticsProviderET> FEngineAnalytics::Analytics;
 TSharedPtr<FEngineSessionManager> FEngineAnalytics::SessionManager;
 
 #if WITH_EDITOR
-static TSharedPtr<FEditorSessionSummaryWriter> SessionSummaryWriter;
-static TSharedPtr<FEditorSessionSummarySender> SessionSummarySender;
+static TUniquePtr<FAnalyticsSessionSummaryManager> AnalyticsSessionSummaryManager;
+static TUniquePtr<FEditorAnalyticsSessionSummary> EditorAnalyticSessionSummary;
+static TSharedPtr<FAnalyticsSessionSummarySender> AnalyticsSessionSummarySender;
 #endif
 
 static TSharedPtr<IAnalyticsProviderET> CreateEpicAnalyticsProvider()
@@ -143,18 +147,26 @@ void FEngineAnalytics::Initialize()
 		}
 
 #if WITH_EDITOR
-		if (!SessionSummaryWriter.IsValid())
+		if (!AnalyticsSessionSummaryManager)
 		{
-			SessionSummaryWriter = MakeShared<FEditorSessionSummaryWriter>(FGenericCrashContext::GetOutOfProcessCrashReporterProcessId());
-			SessionSummaryWriter->Initialize();
-		}
-
-		if (!SessionSummarySender.IsValid())
-		{
-			// if we're using out-of-process crash reporting, then we don't need to create a sender in this process.
-			if (!FGenericCrashContext::IsOutOfProcessCrashReporter())
+			// Create the session summary manager for the Editor instance.
+			AnalyticsSessionSummaryManager = MakeUnique<FAnalyticsSessionSummaryManager>(
+				TEXT("Editor"), // The tag name of the process.
+				FApp::GetInstanceId().ToString(EGuidFormats::Digits), // Unique key to link the principal process (Editor) with subsidiary processes (CRC), that key is passed to CRC.
+				Analytics->GetUserID(),
+				Analytics->GetAppID(),
+				Analytics->GetAppVersion(),
+				Analytics->GetSessionID());
+
+			// The sender will sends orphans sessions and maybe this session if CRC dies first.
+			AnalyticsSessionSummaryManager->SetSender(MakeShared<FAnalyticsSessionSummarySender>(FEngineAnalytics::GetProvider()));
+
+			// Create a property store file with enough pre-reserved capacity to store the analytics data. This reduce risk of running out of disk later.
+			constexpr uint32 ReservedFileCapacity = 16 * 1024;
+			if (TSharedPtr<IAnalyticsPropertyStore> EditorPropertyStore = AnalyticsSessionSummaryManager->MakeStore(ReservedFileCapacity))
 			{
-				SessionSummarySender = MakeShared<FEditorSessionSummarySender>(FEngineAnalytics::GetProvider(), TEXT("Editor"), FPlatformProcess::GetCurrentProcessId());
+				// Create the object responsible to collect the Editor session properties.
+				EditorAnalyticSessionSummary = MakeUnique<FEditorAnalyticsSessionSummary>(EditorPropertyStore, FGenericCrashContext::GetOutOfProcessCrashReporterProcessId());
 			}
 		}
 #endif
@@ -171,16 +183,29 @@ void FEngineAnalytics::Shutdown(bool bIsEngineShutdown)
 	}
 
 #if WITH_EDITOR
-	if (SessionSummaryWriter.IsValid())
+	if (EditorAnalyticSessionSummary)
+	{
+		EditorAnalyticSessionSummary->Shutdown();
+		EditorAnalyticSessionSummary.Reset();
+	}
+
+	if (AnalyticsSessionSummaryManager)
+	{
+		bool bDiscard = !bIsEngineShutdown; // User toggled the 'Send Data' off.
+		AnalyticsSessionSummaryManager->Shutdown(bDiscard);
+		AnalyticsSessionSummaryManager.Reset();
+	}
+	else
 	{
-		SessionSummaryWriter->Shutdown();
-		SessionSummaryWriter.Reset();
+		// The manager cleans any left-over (crash, power outage) on shutdown  when analytics is on but if off, ensure to clean up what could be left from when it was on.
+		FAnalyticsSessionSummaryManager::CleanupExpiredFiles();
 	}
 
-	if (SessionSummarySender.IsValid())
+	// Clean up the outdated sessions created by the deprecated system. If an older compabile Editor is launched, it can still send them up before they get expired.
+	if (FEditorAnalyticsSession::TryLock())
 	{
-		SessionSummarySender->Shutdown();
-		SessionSummarySender.Reset();
+		FEditorAnalyticsSession::CleanupOutdatedIncompatibleSessions(FTimespan::FromDays(30));
+		FEditorAnalyticsSession::Unlock();
 	}
 #endif
 
@@ -200,14 +225,14 @@ void FEngineAnalytics::Tick(float DeltaTime)
 	}
 
 #if WITH_EDITOR
-	if (SessionSummaryWriter.IsValid())
+	if (EditorAnalyticSessionSummary)
 	{
-		SessionSummaryWriter->Tick(DeltaTime);
+		EditorAnalyticSessionSummary->Tick(DeltaTime);
 	}
 
-	if (SessionSummarySender.IsValid())
+	if (AnalyticsSessionSummaryManager.IsValid())
 	{
-		SessionSummarySender->Tick(DeltaTime);
+		AnalyticsSessionSummaryManager->Tick();
 	}
 #endif
 }
@@ -215,9 +240,9 @@ void FEngineAnalytics::Tick(float DeltaTime)
 void FEngineAnalytics::LowDriveSpaceDetected()
 {
 #if WITH_EDITOR
-	if (SessionSummaryWriter.IsValid())
+	if (EditorAnalyticSessionSummary)
 	{
-		SessionSummaryWriter->LowDriveSpaceDetected();
+		EditorAnalyticSessionSummary->LowDriveSpaceDetected();
 	}
 #endif
 }
diff --git a/Engine/Source/Runtime/Engine/Public/Analytics/EngineAnalyticsSessionSummary.h b/Engine/Source/Runtime/Engine/Public/Analytics/EngineAnalyticsSessionSummary.h
new file mode 100644
--- /dev/null
+++ b/Engine/Source/Runtime/Engine/Public/Analytics/EngineAnalyticsSessionSummary.h
@@ -0,0 +1,113 @@
+// Copyright Epic Games, Inc. All Rights Reserved.
+
+#pragma once
+
+#include "CoreMinimal.h"
+#include "Misc/DateTime.h"
+#include "Misc/Optional.h"
+#include "Misc/Timespan.h"
+
+class IAnalyticsPropertyStore;
+class FAnalyticsSessionSummaryWriter;
+struct FUserActivity;
+
+/**
+ * Collects engine events/stats and stores a summary on disk until reported by the analytics session summary manager.
+ */
+class FEngineAnalyticsSessionSummary
+{
+public:
+	/**
+	 * Constructs an analytics summary collector.
+	 * @param Storage Where to store the engine summary properties that are sent on exit.
+	 * @param MonitorProcessId The process ID of CRC launched at start up in 'monitor mode', zero otherwise.
+	 */
+	FEngineAnalyticsSessionSummary(TSharedPtr<IAnalyticsPropertyStore> Storage, uint32 MonitorProcessId);
+
+	/** Destructor. */
+	virtual ~FEngineAnalyticsSessionSummary() = default;
+
+	/** Ticks the session. This calls UpdateSessionProgress() at the configured update period which calls UpdateSessionProgressInternal() for derived classes. */
+	void Tick(float DeltaTime);
+
+	/** Shuts down the session, unregistering the callbacks and closing the session. */
+	void Shutdown();
+
+	/** Set the period used to persist the properties to disk. Some events ignore this setting and persist the properties right away. */
+	void SetPersistPeriod(const FTimespan& Period) { PersistPeriod = Period; }
+
+	/** Get the period used to update and persist the properties to disk. */
+	FTimespan GetPersistPeriod() const { return PersistPeriod; }
+
+	/** Invoked by the engine when the user is running low on disk space.. */
+	void LowDriveSpaceDetected();
+
+protected:
+	/** Invoked during Shutdown() to let derived classes hook in the shutdown. Default implementation does nothing.*/
+	virtual void ShutdownInternal() {}
+
+	/** Invoked during UpdateSessionProgress() to let derived classes hook in the update. Default implementation does nothing. Returns true to force persisting the properties to disk now. */
+	virtual bool UpdateSessionProgressInternal(bool bCrashing) { return false; }
+
+	/** Returns the property store used by the summary. */
+	IAnalyticsPropertyStore* GetStore() { return Store.Get(); }
+
+private:
+	/** Invoked periodically and during important events to update the session progression.*/
+	void UpdateSessionProgress(bool bFlushAsync = true, const FTimespan& FlushTimeout = FTimespan::Zero(), bool bCrashing = false);
+
+	/** Invoked when the engine is crashing.. */
+	void OnCrashing();
+
+	/** Invoked when the engine is requested to terminate. */
+	void OnTerminate();
+
+	/** Invoked when the user is logging in or logging out. Only handles when the user logs out. */
+	void OnUserLoginChanged(bool bLoggingIn, int32, int32);
+
+	/** Invoked when the user activity changes. */
+	void OnUserActivity(const FUserActivity& UserActivity);
+
+	/** Invoked if the engine detects some local changes.. */
+	void OnVanillaStateChanged(bool bIsVanilla);
+
+	/** Check on CRC and return true if its running state state changed (so that the change is persisted on disk). */
+	bool UpdateExternalProcessReporterState(bool bQuickCheck);
+
+	/** Returns whether the debugger state changed. */
+	bool UpdateDebuggerStates();
+
+private:
+	/** The store used to persist the session properties until they are sent. */
+	TSharedPtr<IAnalyticsPropertyStore> Store;
+
+	/** Session timestamp from FDateTime::UtcNow(). Unreliable if user change system date/time (daylight saving or user altering it). */
+	FDateTime SessionStartTimeUtc;
+
+	/** Session timestamp from FPlatformTime::Seconds(). Lose precision when computing long time spans (+/- couple of seconds over a day). */
+	double SessionStartTimeSecs;
+
+	/** The number of time the summary was ticked. */
+	uint64 CurrSessionTickCount = 0;
+
+	/** The monitor process id. CrashReportClient (CRC) acts as the monitoring/reporting process. */
+	uint32 CrcProcessId = 0;
+
+	/** The nominal period at which the session properties are persisted to disk. */
+	FTimespan PersistPeriod;
+
+	/** The next time the session properties should be persisted to disk. */
+	double NextPersistTimeSeconds = 0.0;
+
+	/** Cache CRC exit code (if CRC was spawned in background and unexpectedly died before the engine). */
+	TOptional<int32> CrcExitCode;
+
+	/** Indicates whether Shutdown() method was called. */
+	bool bShutdown = false;
+
+	/** Indicates whether the setting to ignore the debugger was set. */
+	bool bDebuggerIgnored = false;
+
+	/** Indicates whether the presence of the debugger was detected. */
+	bool bWasEverDebugged = false;
+};
