Fixed CIS for error reporter in AnalyticsPropertyStoreTests.cpp

diff --git a/Engine/Source/Runtime/Analytics/Analytics/Private/Tests/AnalyticsPropertyStoreTests.cpp b/Engine/Source/Runtime/Analytics/Analytics/Private/Tests/AnalyticsPropertyStoreTests.cpp
--- a/Engine/Source/Runtime/Analytics/Analytics/Private/Tests/AnalyticsPropertyStoreTests.cpp
+++ b/Engine/Source/Runtime/Analytics/Analytics/Private/Tests/AnalyticsPropertyStoreTests.cpp
@@ -2,8 +2,10 @@
 
 #include "CoreMinimal.h"
 #include "GenericPlatform/GenericPlatformFile.h"
+#include "HAL/FileManager.h"
 #include "HAL/PlatformFileManager.h"
 #include "Misc/AutomationTest.h"
+#include "Misc/Paths.h"
 #include "AnalyticsPropertyStore.h"
 #include "CoreGlobals.h"
 
