Renamed ToString(EAnalyticsSessionShutdownType) as LexToString(EAnalyticsSessionShutdownType) and made it public.

diff --git a/Engine/Source/Runtime/Analytics/AnalyticsET/Private/AnalyticsSessionSummaryManager.cpp b/Engine/Source/Runtime/Analytics/AnalyticsET/Private/AnalyticsSessionSummaryManager.cpp
--- a/Engine/Source/Runtime/Analytics/AnalyticsET/Private/AnalyticsSessionSummaryManager.cpp
+++ b/Engine/Source/Runtime/Analytics/AnalyticsET/Private/AnalyticsSessionSummaryManager.cpp
@@ -135,32 +135,33 @@ namespace AnalyticsSessionUtils
 		}
 	}
 
-	/** Convert the shutdown type to the string representation known by the analytics backend. */
-	const TCHAR* ToString(EAnalyticsSessionShutdownType Code)
+} // namespace AnalyticsSessionUtils
+
+
+const TCHAR* LexToString(EAnalyticsSessionShutdownType ShutdownTypeCode)
+{
+	switch (ShutdownTypeCode)
 	{
-		switch (Code)
-		{
-		case EAnalyticsSessionShutdownType::Shutdown:
-			return TEXT("Shutdown");
+	case EAnalyticsSessionShutdownType::Shutdown:
+		return TEXT("Shutdown");
 
-		case EAnalyticsSessionShutdownType::Crashed:
-			return TEXT("Crashed");
+	case EAnalyticsSessionShutdownType::Crashed:
+		return TEXT("Crashed");
 
-		case EAnalyticsSessionShutdownType::Terminated:
-			return TEXT("Terminated");
+	case EAnalyticsSessionShutdownType::Terminated:
+		return TEXT("Terminated");
 
-		case EAnalyticsSessionShutdownType::Debugged:
-			return TEXT("Debugger");
+	case EAnalyticsSessionShutdownType::Debugged:
+		return TEXT("Debugger");
 
-		case EAnalyticsSessionShutdownType::Abnormal:
-			return TEXT("AbnormalShutdown");
+	case EAnalyticsSessionShutdownType::Abnormal:
+		return TEXT("AbnormalShutdown");
 
-		default:
-			return TEXT("Unknown");
-		}
+	case EAnalyticsSessionShutdownType::Unknown:
+	default:
+		return TEXT("Unknown");
 	}
-
-} // namespace AnalyticsSessionUtils
+}
 
 
 const TAnalyticsProperty<int32> FAnalyticsSessionSummaryManager::ShutdownTypeCodeProperty = TEXT("ShutdownTypeCode");
@@ -495,7 +496,7 @@ bool FAnalyticsSessionSummaryManager::AggregateSummaries(const FString& InProces
 				if (!OutSummaryProperties.Contains(AnalyticsManagerProperties::ShutdownType.Key))
 				{
 					// Convert 'ShutdownTypeCode' into its string representation and emit it as 'ShutdownType' as known by the analytics backend.
-					OutSummaryProperties.Emplace(AnalyticsManagerProperties::ShutdownType.Key, AnalyticsSessionUtils::ToString(static_cast<EAnalyticsSessionShutdownType>(ShutdownTypeCode)));
+					OutSummaryProperties.Emplace(AnalyticsManagerProperties::ShutdownType.Key, LexToString(static_cast<EAnalyticsSessionShutdownType>(ShutdownTypeCode)));
 				}
 			}
 
diff --git a/Engine/Source/Runtime/Analytics/AnalyticsET/Public/AnalyticsSessionSummaryManager.h b/Engine/Source/Runtime/Analytics/AnalyticsET/Public/AnalyticsSessionSummaryManager.h
--- a/Engine/Source/Runtime/Analytics/AnalyticsET/Public/AnalyticsSessionSummaryManager.h
+++ b/Engine/Source/Runtime/Analytics/AnalyticsET/Public/AnalyticsSessionSummaryManager.h
@@ -19,6 +19,11 @@ enum class EAnalyticsSessionShutdownType : int32
 	Unknown,
 };
 
+/**
+ * Converts the specified shutdown code to its string representation.
+ */
+ANALYTICSET_API const TCHAR* LexToString(EAnalyticsSessionShutdownType ShutdownTypeCode);
+
 /**
  * Manages the lifecycle of analytics summaries created by an application and its subsidiary processes. It manages the
  * summary files created by the processes, tracks which processes are working together, aggregates the summaries
