Fixed analytics summary system losing type information.
  - Instead of converting summary properties to string duing session aggregation, put properties in FAnalyticsEventAttribute instances to preserve the type information.
Added MonitorEngineVersion, MonitorReportCount, MonitorEnsureCount, MonitorAssertCount to CRC analytics summary.
Added MissingDataFrom to the analytics summary when the property file from a collaboration process failed to load.
Fixed the analytics summary manager to aggreage and produce a report even if the helper process data couldn't be loaded.

diff --git a/Engine/Source/Programs/CrashReportClient/Private/CrashReportAnalyticsSessionSummary.cpp b/Engine/Source/Programs/CrashReportClient/Private/CrashReportAnalyticsSessionSummary.cpp
--- a/Engine/Source/Programs/CrashReportClient/Private/CrashReportAnalyticsSessionSummary.cpp
+++ b/Engine/Source/Programs/CrashReportClient/Private/CrashReportAnalyticsSessionSummary.cpp
@@ -11,6 +11,7 @@
 #include "GenericPlatform/GenericPlatformCrashContext.h"
 #include "HAL/PlatformTime.h"
 #include "Logging/LogMacros.h"
+#include "Misc/EngineVersion.h"
 #include "Misc/Paths.h"
 #include "Misc/ScopeLock.h"
 #include "Templates/UnrealTemplate.h"
@@ -43,11 +44,18 @@ LRESULT CALLBACK CrashReportAnalyticsSessionSummaryWindowProc(HWND Hwnd, UINT Ms
 
 namespace CrcAnalyticsProperties
 {
+	//NOTE: Update this when you add/remove/change key behavior.
+	constexpr uint32 CrcAnalyticsSummaryVersion = 1;
+
 	/** The exit code of the monitored application. */
 	static const TAnalyticsProperty<int32> MonitoredAppExitCode(TEXT("ExitCode"));
 	/** Track when CRC detected the death of the monitored app. */
 	static const TAnalyticsProperty<FDateTime> MonitoredAppDeathTimestamp(TEXT("DeathTimestamp"));
 
+	/** CRC engine version. In-dev, people don't always recompile CRC and we get disparity between the monitored app and CRC version. */
+	static const TAnalyticsProperty<FString> EngineVersion(TEXT("MonitorEngineVersion"));
+	/** The version number of the key/set used by CRC. */
+	static const TAnalyticsProperty<uint32> SummaryVersionNumber(TEXT("MonitorSummaryVersion"));
 	/** The CRC startup timestamp. */
 	static const TAnalyticsProperty<FDateTime> StartupTimestamp(TEXT("MonitorStartupTimestamp"));
 	/** The CRC timestamp. */
@@ -76,6 +84,12 @@ namespace CrcAnalyticsProperties
 	static const TAnalyticsProperty<bool> IsCrashing(TEXT("MonitorCrashed"));
 	/** Whether CRC was shutdown. */
 	static const TAnalyticsProperty<bool> WasShutdown(TEXT("MonitorWasShutdown"));
+	/** Number of crash event passed to CRC. (Ensure, Assert, Crash, etc). */
+	static const TAnalyticsProperty<uint32> ReportCount(TEXT("MonitorReportCount"));
+	/** Number of ensures handled by CRC. */
+	static const TAnalyticsProperty<uint32> EnsureCount(TEXT("MonitorEnsureCount"));
+	/** Number of assert handled by CRC.*/
+	static const TAnalyticsProperty<uint32> AssertCount(TEXT("MonitorAssertCount"));
 }
 
 namespace CrashReportClientUtils
@@ -184,7 +198,7 @@ public:
 	{
 	}
 
-	virtual bool SendSessionSummary(const FString& UserId, const FString& AppId, const FString& AppVersion, const FString& SessionId, const TMap<FString, FString>& Properties) override
+	virtual bool SendSessionSummary(const FString& UserId, const FString& AppId, const FString& AppVersion, const FString& SessionId, const TArray<FAnalyticsEventAttribute>& Properties) override
 	{
 		// CRC should only send one session (its own), but do a reset in case this convention changes.
 		bAbnormalShutdown = false;
@@ -192,14 +206,14 @@ public:
 
 		// Analyze the report to be sent and try to figure out if this is an abnormal shutdown. They keys are taken from the engine analytics session summary. To prevent dependencies
 		// between CRC and Engine analytics, we duplicate the keys here.
-		if (const FString* ShutdownTypeCode = Properties.Find(FAnalyticsSessionSummaryManager::ShutdownTypeCodeProperty.Key))
+		if (const FAnalyticsEventAttribute* ShutdownTypeCode = Properties.FindByPredicate([](const FAnalyticsEventAttribute& Candidate) { return Candidate.GetName() == FAnalyticsSessionSummaryManager::ShutdownTypeCodeProperty.Key; }))
 		{
-			bAbnormalShutdown = (*ShutdownTypeCode == LexToString((int32)EAnalyticsSessionShutdownType::Abnormal));
+			bAbnormalShutdown = (ShutdownTypeCode->GetValue() == LexToString((int32)EAnalyticsSessionShutdownType::Abnormal));
 			if (bAbnormalShutdown)
 			{
-				if (const FString* UserLoggingOut = Properties.Find(FAnalyticsSessionSummaryManager::IsUserLoggingOutProperty.Key))
+				if (const FAnalyticsEventAttribute* UserLoggingOut = Properties.FindByPredicate([](const FAnalyticsEventAttribute& Candidate) { return Candidate.GetName() == FAnalyticsSessionSummaryManager::IsUserLoggingOutProperty.Key; }))
 				{
-					bUserLoggingOut = (*UserLoggingOut == LexToString(true));
+					bUserLoggingOut = (UserLoggingOut->GetValue() == LexToString(true));
 				}
 			}
 		}
@@ -255,6 +269,8 @@ void FCrashReportAnalyticsSessionSummary::Initialize(const FString& ProcessGroup
 				FCoreDelegates::ApplicationWillTerminateDelegate.AddRaw(this, &FCrashReportAnalyticsSessionSummary::OnApplicationWillTerminate);
 				FCoreDelegates::OnHandleSystemError.AddRaw(this, &FCrashReportAnalyticsSessionSummary::OnHandleSystemError);
 
+				CrcAnalyticsProperties::EngineVersion.Set(PropertyStore.Get(), FEngineVersion::Current().ToString(EVersionComponent::Changelist));
+				CrcAnalyticsProperties::SummaryVersionNumber.Set(PropertyStore.Get(), CrcAnalyticsProperties::CrcAnalyticsSummaryVersion);
 				CrcAnalyticsProperties::StartupTimestamp.Set(PropertyStore.Get(), FDateTime::UtcNow());
 				CrcAnalyticsProperties::Timestamp.Set(PropertyStore.Get(), FDateTime::UtcNow());
 				CrcAnalyticsProperties::SessionDurationSecs.Set(PropertyStore.Get(), FMath::FloorToInt(static_cast<float>(FPlatformTime::Seconds() - SessionStartTimeSecs)));
@@ -266,6 +282,9 @@ void FCrashReportAnalyticsSessionSummary::Initialize(const FString& ProcessGroup
 				CrcAnalyticsProperties::IsCrashing.Set(PropertyStore.Get(), false);
 				CrcAnalyticsProperties::WasShutdown.Set(PropertyStore.Get(), false);
 				CrcAnalyticsProperties::QuitSignalRecv.Set(PropertyStore.Get(), false);
+				CrcAnalyticsProperties::ReportCount.Set(PropertyStore.Get(), 0);
+				CrcAnalyticsProperties::EnsureCount.Set(PropertyStore.Get(), 0);
+				CrcAnalyticsProperties::AssertCount.Set(PropertyStore.Get(), 0);
 
 				UpdatePowerStatus();
 				Flush();
@@ -551,15 +570,18 @@ void FCrashReportAnalyticsSessionSummary::OnCrashReportStarted(ECrashContextType
 	CrashReportStartTimeSecs = FPlatformTime::Seconds();
 
 	CrcAnalyticsProperties::IsReportingCrash.Set(PropertyStore.Get(), true);
+	CrcAnalyticsProperties::ReportCount.Update(PropertyStore.Get(), [](uint32& Actual) { ++Actual; return true; });
 	LogEvent(FString::Printf(TEXT("Report/Start:%s"), *FDateTime::UtcNow().ToString()));
 
 	// Log the assert and ensure condition/file/line/message to the diagnostic log gathered by the analytics to enable searching/grouping them later on.
 	if (CrashType == ECrashContextType::Assert)
 	{
+		CrcAnalyticsProperties::AssertCount.Update(PropertyStore.Get(), [](uint32& Actual) { ++Actual; return true; });
 		LogEvent(FString::Printf(TEXT("Assert/Msg: %s"), ErrorMsg));
 	}
 	else if (CrashType == ECrashContextType::Ensure)
 	{
+		CrcAnalyticsProperties::EnsureCount.Update(PropertyStore.Get(), [](uint32& Actual) { ++Actual; return true; });
 		LogEvent(FString::Printf(TEXT("Ensure/Msg: %s"), ErrorMsg));
 	}
 }
diff --git a/Engine/Source/Runtime/Analytics/Analytics/Private/AnalyticsPropertyStore.cpp b/Engine/Source/Runtime/Analytics/Analytics/Private/AnalyticsPropertyStore.cpp
--- a/Engine/Source/Runtime/Analytics/Analytics/Private/AnalyticsPropertyStore.cpp
+++ b/Engine/Source/Runtime/Analytics/Analytics/Private/AnalyticsPropertyStore.cpp
@@ -125,8 +125,7 @@ bool FAnalyticsPropertyStore::Create(const FString& Pathname, uint32 CapacityHin
 	// Reserve the space in the file. Nothing else until the store is flushed the first time.
 	uint8 Dummy = 0;
 	FileHandle->Seek(ReservedSize - sizeof(Dummy));
-	FileHandle->Write(&Dummy, sizeof(Dummy));
-	return true;
+	return FileHandle->Write(&Dummy, sizeof(Dummy));
 }
 
 bool FAnalyticsPropertyStore::Load(const FString& Pathname)
@@ -805,7 +804,7 @@ void FAnalyticsPropertyStore::Defragment()
 	bFragmented = false;
 }
 
-void FAnalyticsPropertyStore::VisitAll(const TFunction<void(const FString&, FString&&)>& VisitFn) const
+void FAnalyticsPropertyStore::VisitAll(const TFunction<void(FAnalyticsEventAttribute&&)>& VisitFn) const
 {
 	FScopeLock Lock(&StoreLock);
 
@@ -820,7 +819,7 @@ void FAnalyticsPropertyStore::VisitAll(const TFunction<void(const FString&, FStr
 			{
 				int32 V;
 				StorageReader.Serialize(&V, sizeof(V));
-				VisitFn(NameOffsetPair.Key, LexToString(V));
+				VisitFn(FAnalyticsEventAttribute(NameOffsetPair.Key, V));
 				break;
 			}
 
@@ -828,7 +827,7 @@ void FAnalyticsPropertyStore::VisitAll(const TFunction<void(const FString&, FStr
 			{
 				uint32 V;
 				StorageReader.Serialize(&V, sizeof(V));
-				VisitFn(NameOffsetPair.Key, LexToString(V));
+				VisitFn(FAnalyticsEventAttribute(NameOffsetPair.Key, V));
 				break;
 			}
 
@@ -836,7 +835,7 @@ void FAnalyticsPropertyStore::VisitAll(const TFunction<void(const FString&, FStr
 			{
 				int64 V;
 				StorageReader.Serialize(&V, sizeof(V));
-				VisitFn(NameOffsetPair.Key, LexToString(V));
+				VisitFn(FAnalyticsEventAttribute(NameOffsetPair.Key, V));
 				break;
 			}
 
@@ -844,7 +843,7 @@ void FAnalyticsPropertyStore::VisitAll(const TFunction<void(const FString&, FStr
 			{
 				uint64 V;
 				StorageReader.Serialize(&V, sizeof(V));
-				VisitFn(NameOffsetPair.Key, LexToString(V));
+				VisitFn(FAnalyticsEventAttribute(NameOffsetPair.Key, V));
 				break;
 			}
 
@@ -852,7 +851,7 @@ void FAnalyticsPropertyStore::VisitAll(const TFunction<void(const FString&, FStr
 			{
 				float V;
 				StorageReader.Serialize(&V, sizeof(V));
-				VisitFn(NameOffsetPair.Key, FString::SanitizeFloat(V));
+				VisitFn(FAnalyticsEventAttribute(NameOffsetPair.Key, V));
 				break;
 			}
 
@@ -860,7 +859,7 @@ void FAnalyticsPropertyStore::VisitAll(const TFunction<void(const FString&, FStr
 			{
 				double V;
 				StorageReader.Serialize(&V, sizeof(V));
-				VisitFn(NameOffsetPair.Key, FString::SanitizeFloat(V));
+				VisitFn(FAnalyticsEventAttribute(NameOffsetPair.Key, V));
 				break;
 			}
 
@@ -868,23 +867,23 @@ void FAnalyticsPropertyStore::VisitAll(const TFunction<void(const FString&, FStr
 			{
 				bool V;
 				StorageReader.Serialize(&V, sizeof(V));
-				VisitFn(NameOffsetPair.Key, LexToString(V)); // Convert to "true" or "false".
+				VisitFn(FAnalyticsEventAttribute(NameOffsetPair.Key, V)); // Convert to "true" or "false".
 				break;
 			}
 
 			case ETypeCode::Str:
 			{
-				FString Value;
-				GetStringValueInternal(NameOffsetPair.Key, Value);
-				VisitFn(NameOffsetPair.Key, MoveTemp(Value));
+				FString V;
+				GetStringValueInternal(NameOffsetPair.Key, V);
+				VisitFn(FAnalyticsEventAttribute(NameOffsetPair.Key, V));
 				break;
 			}
 
 			case ETypeCode::Date:
 			{
-				FDateTime Value;
-				GetDateTimeValueInternal(NameOffsetPair.Key, Value);
-				VisitFn(NameOffsetPair.Key, Value.ToIso8601()); // Iso-8601 is the format used by Analytics in general for date/time.
+				FDateTime V;
+				GetDateTimeValueInternal(NameOffsetPair.Key, V);
+				VisitFn(FAnalyticsEventAttribute(NameOffsetPair.Key, V.ToIso8601())); // Iso-8601 is the format used by Analytics in general for date/time.
 				break;
 			}
 
diff --git a/Engine/Source/Runtime/Analytics/Analytics/Private/Tests/AnalyticsPropertyStoreTests.cpp b/Engine/Source/Runtime/Analytics/Analytics/Private/Tests/AnalyticsPropertyStoreTests.cpp
--- a/Engine/Source/Runtime/Analytics/Analytics/Private/Tests/AnalyticsPropertyStoreTests.cpp
+++ b/Engine/Source/Runtime/Analytics/Analytics/Private/Tests/AnalyticsPropertyStoreTests.cpp
@@ -175,47 +175,47 @@ bool FAnalyticsPropertyStoreAutomationTest::RunTest(const FString& Parameters)
 		check(Store.Get(BoolKey, Actual_Bool) == IAnalyticsPropertyStore::EStatusCode::Success && Actual_Bool == Expected_Bool);
 
 		// Test string conversion.
-		Store.VisitAll([&](const FString& Name, FString&& Value)
+		Store.VisitAll([&](FAnalyticsEventAttribute&& Attr)
 		{
-			if (Name == I32Key)
+			if (Attr.GetName() == I32Key)
 			{
-				check(Value == LexToString(Expected_I32));
+				check(Attr.GetValue() == LexToString(Expected_I32));
 			}
-			else if (Name == U32Key)
+			else if (Attr.GetName() == U32Key)
 			{
-				check(Value == LexToString(Expected_U32));
+				check(Attr.GetValue() == LexToString(Expected_U32));
 			}
-			else if (Name == I64Key)
+			else if (Attr.GetName() == I64Key)
 			{
-				check(Value == LexToString(Expected_I64));
+				check(Attr.GetValue() == LexToString(Expected_I64));
 			}
-			else if (Name == U64Key)
+			else if (Attr.GetName() == U64Key)
 			{
-				check(Value == LexToString(Expected_U64));
+				check(Attr.GetValue() == LexToString(Expected_U64));
 			}
-			else if (Name == FltKey)
+			else if (Attr.GetName() == FltKey)
 			{
-				check(Value == FString::SanitizeFloat(Expected_Flt));
+				check(Attr.GetValue() == FString::SanitizeFloat(Expected_Flt));
 			}
-			else if (Name == DblKey)
+			else if (Attr.GetName() == DblKey)
 			{
-				check(Value == FString::SanitizeFloat(Expected_Dbl));
+				check(Attr.GetValue() == FString::SanitizeFloat(Expected_Dbl));
 			}
-			else if (Name == BoolKey)
+			else if (Attr.GetName() == BoolKey)
 			{
-				check(Value == LexToString(Expected_Bool));
+				check(Attr.GetValue() == LexToString(Expected_Bool));
 			}
-			else if (Name == Str1Key)
+			else if (Attr.GetName() == Str1Key)
 			{
-				check(Value == Expected_Str1);
+				check(Attr.GetValue() == Expected_Str1);
 			}
-			else if (Name == DateKey)
+			else if (Attr.GetName() == DateKey)
 			{
-				check(Value == Expected_Date.ToIso8601());
+				check(Attr.GetValue() == Expected_Date.ToIso8601());
 			}
-			else if (Name == Str2Key)
+			else if (Attr.GetName() == Str2Key)
 			{
-				check(Value == Expected_Str2);
+				check(Attr.GetValue() == Expected_Str2);
 			}
 			else
 			{
diff --git a/Engine/Source/Runtime/Analytics/Analytics/Public/AnalyticsPropertyStore.h b/Engine/Source/Runtime/Analytics/Analytics/Public/AnalyticsPropertyStore.h
--- a/Engine/Source/Runtime/Analytics/Analytics/Public/AnalyticsPropertyStore.h
+++ b/Engine/Source/Runtime/Analytics/Analytics/Public/AnalyticsPropertyStore.h
@@ -93,7 +93,7 @@ public:
 	virtual EStatusCode Get(const FString& Key, bool& OutValue) const override;
 	virtual EStatusCode Get(const FString& Key, FString& OutValue) const override;
 	virtual EStatusCode Get(const FString& Key, FDateTime& OutValue) const override;
-	virtual void VisitAll(const TFunction<void(const FString& /*Key*/, FString&& /*Value*/)>& VisitFn) const override;
+	virtual void VisitAll(const TFunction<void(FAnalyticsEventAttribute&&)>& VisitFn) const override;
 
 	/**
 	 * Flushes cached values to disk. Only one thread can flush at a time, causing possible latency if a flush is alreary happenning.
diff --git a/Engine/Source/Runtime/Analytics/Analytics/Public/Interfaces/IAnalyticsPropertyStore.h b/Engine/Source/Runtime/Analytics/Analytics/Public/Interfaces/IAnalyticsPropertyStore.h
--- a/Engine/Source/Runtime/Analytics/Analytics/Public/Interfaces/IAnalyticsPropertyStore.h
+++ b/Engine/Source/Runtime/Analytics/Analytics/Public/Interfaces/IAnalyticsPropertyStore.h
@@ -5,6 +5,7 @@
 #include "CoreMinimal.h"
 #include "Misc/DateTime.h"
 #include "Templates/Function.h"
+#include "AnalyticsEventAttribute.h"
 
 /**
  * Implements a fast type safe key/value database to store analytics properties collected during a session.
@@ -132,7 +133,7 @@ public:
 	 * Iterates the keys currently stored and invokes the visitor function for each key, converting the value to its string representation.
 	 * @param VisitFn The callback invoked for each key/value pair visited.
 	 */
-	virtual void VisitAll(const TFunction<void(const FString& /*Key*/, FString&& /*Value*/)>& VisitFn) const = 0;
+	virtual void VisitAll(const TFunction<void(FAnalyticsEventAttribute&&)>& VisitFn) const = 0;
 
 	/**
 	 * Flushes cached values to persistent storage.
diff --git a/Engine/Source/Runtime/Analytics/AnalyticsET/Private/AnalyticsSessionSummaryManager.cpp b/Engine/Source/Runtime/Analytics/AnalyticsET/Private/AnalyticsSessionSummaryManager.cpp
--- a/Engine/Source/Runtime/Analytics/AnalyticsET/Private/AnalyticsSessionSummaryManager.cpp
+++ b/Engine/Source/Runtime/Analytics/AnalyticsET/Private/AnalyticsSessionSummaryManager.cpp
@@ -25,15 +25,18 @@ namespace AnalyticsManagerProperties
 	static const TAnalyticsProperty<bool>    InternalWasProcessed      = TEXT("Internal.WasProcessed");
 
 	// Those values are implicitely added by the manager into the report because the analytics backend expecpts them.
-	static const TAnalyticsProperty<FString> ShutdownType      = TEXT("ShutdownType");
-	static const TAnalyticsProperty<bool>    DelayedSend       = TEXT("DelayedSend");
-	static const TAnalyticsProperty<FString> SentFrom          = TEXT("SentFrom");
+	static const TAnalyticsProperty<FString> SessionId       = TEXT("SessionId");
+	static const TAnalyticsProperty<FString> ShutdownType    = TEXT("ShutdownType");
+	static const TAnalyticsProperty<bool>    DelayedSend     = TEXT("DelayedSend");
+	static const TAnalyticsProperty<FString> SentFrom        = TEXT("SentFrom");
+	static const TAnalyticsProperty<FString> MissingDataFrom = TEXT("MissingDataFrom");
 
 	// List of reserved property key.
 	static TSet<FString> ReservedKeys = {
 			AnalyticsManagerProperties::ShutdownType.Key,
 			AnalyticsManagerProperties::DelayedSend.Key,
-			AnalyticsManagerProperties::SentFrom.Key
+			AnalyticsManagerProperties::SentFrom.Key,
+			AnalyticsManagerProperties::MissingDataFrom.Key
 		};
 
 	bool IsReserved(const FString& Key)
@@ -426,9 +429,9 @@ void FAnalyticsSessionSummaryManager::ProcessSummary(const FString& InProcessGro
 	}
 
 	// Aggregate the principal and subsidiary process summaries.
-	TMap<FString, FString> InternalProperties;
-	TMap<FString, FString> SummaryProperties;
-	if (!AggregateSummaries(InProcessGroupId, InProcessGroup.PropertyFiles, SummaryProperties, InternalProperties))
+	TArray<FAnalyticsEventAttribute> InternalProperties;
+	TArray<FAnalyticsEventAttribute> SummaryProperties;
+	if (!AggregateSummaries(InProcessGroupId, InProcessGroup, SummaryProperties, InternalProperties))
 	{
 		CleanupFiles(InProcessGroup.PropertyFiles, /*bOnSuccess*/false);
 		return;
@@ -445,15 +448,15 @@ void FAnalyticsSessionSummaryManager::ProcessSummary(const FString& InProcessGro
 	else
 	{
 		// The principal process saves those internal properties to let a subsidiary or another principal processes submit the summary on its behalf.
-		const FString* UserIdProp = InternalProperties.Find(AnalyticsManagerProperties::InternalSessionUserId.Key);
-		const FString* AppIdProp = InternalProperties.Find(AnalyticsManagerProperties::InternalSessionAppId.Key);
-		const FString* AppVersionProp = InternalProperties.Find(AnalyticsManagerProperties::InternalSessionAppVersion.Key);
-		const FString* SessionIdProp = InternalProperties.Find(AnalyticsManagerProperties::InternalSessionId.Key);
+		const FAnalyticsEventAttribute* UserIdProp = InternalProperties.FindByPredicate([](const FAnalyticsEventAttribute& Candidate) { return Candidate.GetName() == AnalyticsManagerProperties::InternalSessionUserId.Key; });
+		const FAnalyticsEventAttribute* AppIdProp = InternalProperties.FindByPredicate([](const FAnalyticsEventAttribute& Candidate) { return Candidate.GetName() == AnalyticsManagerProperties::InternalSessionAppId.Key; });
+		const FAnalyticsEventAttribute* AppVersionProp = InternalProperties.FindByPredicate([](const FAnalyticsEventAttribute& Candidate) { return Candidate.GetName() == AnalyticsManagerProperties::InternalSessionAppVersion.Key; });
+		const FAnalyticsEventAttribute* SessionIdProp = InternalProperties.FindByPredicate([](const FAnalyticsEventAttribute& Candidate) { return Candidate.GetName() == AnalyticsManagerProperties::InternalSessionId.Key; });
 
 		if (UserIdProp && AppIdProp && AppVersionProp && SessionIdProp)
 		{
 			CleanupFiles(InProcessGroup.PropertyFiles, /*bOnSuccess*/true);
-			SummarySender->SendSessionSummary(*UserIdProp, *AppIdProp, *AppVersionProp, *SessionIdProp, SummaryProperties);
+			SummarySender->SendSessionSummary(UserIdProp->GetValue(), AppIdProp->GetValue(), AppVersionProp->GetValue(), SessionIdProp->GetValue(), SummaryProperties);
 		}
 		else
 		{
@@ -462,11 +465,13 @@ void FAnalyticsSessionSummaryManager::ProcessSummary(const FString& InProcessGro
 	}
 }
 
-bool FAnalyticsSessionSummaryManager::AggregateSummaries(const FString& InProcessGroupId, const TArray<FPropertyFileInfo>& PropertyFiles, TMap<FString, FString>& OutSummaryProperties, TMap<FString, FString>& OutInternalProperties)
+bool FAnalyticsSessionSummaryManager::AggregateSummaries(const FString& InProcessGroupId, const FProcessGroup& ProcessGroup, TArray<FAnalyticsEventAttribute>& OutSummaryProperties, TArray<FAnalyticsEventAttribute>& OutInternalProperties)
 {
 	FAnalyticsPropertyStore PropertyStore;
 
-	auto RenameKey = [](const FString& KeyToRename, const FString& FromProcess, const TMap<FString, FString>& ExistingKeys, const TCHAR* Reason)
+	TSet<FString> SummaryKeys;
+
+	auto RenameKey = [](const FString& KeyToRename, const FString& FromProcess, const TSet<FString>& ExistingKeys, const TCHAR* Reason)
 	{
 		FString RenamedKey;
 		int32 Index = 1;
@@ -479,8 +484,10 @@ bool FAnalyticsSessionSummaryManager::AggregateSummaries(const FString& InProces
 		return RenamedKey;
 	};
 
+	TArray<FString> MissingDataFromProcesses;
+
 	// Merge the properties collected by the principal and its subsidiary processes.
-	for (const FPropertyFileInfo& PropertyFileInfo : PropertyFiles)
+	for (const FPropertyFileInfo& PropertyFileInfo : ProcessGroup.PropertyFiles)
 	{
 		if (PropertyStore.Load(PropertyFileInfo.Pathname))
 		{
@@ -493,32 +500,37 @@ bool FAnalyticsSessionSummaryManager::AggregateSummaries(const FString& InProces
 			if (PropertyStore.Get(ShutdownTypeCodeProperty.Key, ShutdownTypeCode) == IAnalyticsPropertyStore::EStatusCode::Success)
 			{
 				// Don't add 'ShutdownType' more than once. If 'ShutdownTypeCode' key used to derive the 'ShutdownType' is duplicated, the code will emit a warning and rename the second instance.
-				if (!OutSummaryProperties.Contains(AnalyticsManagerProperties::ShutdownType.Key))
+				if (!SummaryKeys.Contains(AnalyticsManagerProperties::ShutdownType.Key))
 				{
 					// Convert 'ShutdownTypeCode' into its string representation and emit it as 'ShutdownType' as known by the analytics backend.
 					OutSummaryProperties.Emplace(AnalyticsManagerProperties::ShutdownType.Key, LexToString(static_cast<EAnalyticsSessionShutdownType>(ShutdownTypeCode)));
+					SummaryKeys.Add(AnalyticsManagerProperties::ShutdownType.Key);
 				}
 			}
 
-			PropertyStore.VisitAll([&PropertyFileInfo, &OutSummaryProperties, &OutInternalProperties, &RenameKey](const FString& InKey, FString&& InValue)
+			PropertyStore.VisitAll([&PropertyFileInfo, &OutSummaryProperties, &OutInternalProperties, &RenameKey, &SummaryKeys](FAnalyticsEventAttribute&& Attr)
 			{
-				if (InKey.StartsWith(AnalyticsManagerProperties::InternalPropertyPrefix))
+				if (Attr.GetName().StartsWith(AnalyticsManagerProperties::InternalPropertyPrefix))
 				{
-					OutInternalProperties.Emplace(InKey, InValue);
+					SummaryKeys.Add(Attr.GetName());
+					OutInternalProperties.Emplace(Attr);
 				}
-				else if (AnalyticsManagerProperties::IsReserved(InKey))
+				else if (AnalyticsManagerProperties::IsReserved(Attr.GetName()))
 				{
-					FString RenamedKey = RenameKey(InKey, PropertyFileInfo.ProcessName, OutSummaryProperties, TEXT("Key name is reserved."));
-					OutSummaryProperties.Emplace(RenamedKey, InValue);
+					FString RenamedKey = RenameKey(Attr.GetName(), PropertyFileInfo.ProcessName, SummaryKeys, TEXT("Key name is reserved."));
+					SummaryKeys.Add(RenamedKey);
+					OutSummaryProperties.Emplace(RenamedKey, Attr.GetValue());
 				}
-				else if (!OutSummaryProperties.Contains(InKey))
+				else if (!SummaryKeys.Contains(Attr.GetName()))
 				{
-					OutSummaryProperties.Emplace(InKey, InValue);
+					SummaryKeys.Add(Attr.GetName());
+					OutSummaryProperties.Emplace(Attr);
 				}
 				else // Deal with key collisions. Keys are supposed to be unique. Duplication should be easy to spot during developement.
 				{
-					FString RenamedKey = RenameKey(InKey, PropertyFileInfo.ProcessName, OutSummaryProperties, TEXT("Key collision between processes of a group."));
-					OutSummaryProperties.Emplace(RenamedKey, InValue);
+					FString RenamedKey = RenameKey(Attr.GetName(), PropertyFileInfo.ProcessName, SummaryKeys, TEXT("Key collision between processes of a group."));
+					SummaryKeys.Add(Attr.GetName());
+					OutSummaryProperties.Emplace(RenamedKey, Attr.GetValue());
 				}
 			});
 
@@ -526,20 +538,32 @@ bool FAnalyticsSessionSummaryManager::AggregateSummaries(const FString& InProces
 			AnalyticsManagerProperties::InternalWasProcessed.Set(&PropertyStore, true);
 			PropertyStore.Flush();
 		}
+		else if (PropertyFileInfo.ProcessId == ProcessGroup.PrincipalProcessId)
+		{
+			// Without the principal process data, too much data is missing, abort.
+			return false;
+		}
 		else
 		{
-			return false; // The file failed to load. Was it deleted? locked? internal checksum failed?
+			// Subsidiary process data will be missing. Not ideal, but this should only be complementary info for the principal process, flag it and continue.
+			MissingDataFromProcesses.Add(PropertyFileInfo.ProcessName);
 		}
 	}
 
 	// The manager saves this key in the store it creates.
-	if (const FString* InternalSessionId = OutInternalProperties.Find(AnalyticsManagerProperties::InternalSessionId.Key))
+	if (const FAnalyticsEventAttribute* InternalSessionId = OutInternalProperties.FindByPredicate([](const FAnalyticsEventAttribute& Candidate) { return Candidate.GetName() == AnalyticsManagerProperties::InternalSessionId.Key; }))
+	{
+		OutSummaryProperties.Emplace(AnalyticsManagerProperties::SessionId.Key, InternalSessionId->GetValue());
+	}
+
+	if (!MissingDataFromProcesses.IsEmpty())
 	{
-		OutSummaryProperties.Emplace(TEXT("SessionId"), *InternalSessionId);
+		// Add the subsidiary process tag name(s) for which the data file couldn't be loaded (file was locked, corrupted, ...)
+		OutSummaryProperties.Emplace(AnalyticsManagerProperties::MissingDataFrom.Key, FString::Join(MissingDataFromProcesses, TEXT(",")));
 	}
 
 	// The summary is sent delayed when it is processed from another group.
-	OutSummaryProperties.Emplace(AnalyticsManagerProperties::DelayedSend.Key, LexToString(InProcessGroupId != ProcessGroupId));
+	OutSummaryProperties.Emplace(AnalyticsManagerProperties::DelayedSend.Key, InProcessGroupId != ProcessGroupId);
 
 	// This process is about to send the report, record the process name in the summary.
 	OutSummaryProperties.Emplace(AnalyticsManagerProperties::SentFrom.Key, ProcessName);
diff --git a/Engine/Source/Runtime/Analytics/AnalyticsET/Private/AnalyticsSessionSummarySender.cpp b/Engine/Source/Runtime/Analytics/AnalyticsET/Private/AnalyticsSessionSummarySender.cpp
--- a/Engine/Source/Runtime/Analytics/AnalyticsET/Private/AnalyticsSessionSummarySender.cpp
+++ b/Engine/Source/Runtime/Analytics/AnalyticsET/Private/AnalyticsSessionSummarySender.cpp
@@ -9,13 +9,13 @@
 DEFINE_LOG_CATEGORY_STATIC(LogAnalyticsSessionSummarySender, Verbose, All);
 
 
-FAnalyticsSessionSummarySender::FAnalyticsSessionSummarySender(IAnalyticsProviderET& Provider, TFunction<bool(const FString&, const FString&)> ShouldEmitFilterFunc)
+FAnalyticsSessionSummarySender::FAnalyticsSessionSummarySender(IAnalyticsProviderET& Provider, TFunction<bool(const FAnalyticsEventAttribute&)> ShouldEmitFilterFunc)
 	: AnalyticsProvider(Provider)
 	, ShouldEmitPropFunc(MoveTemp(ShouldEmitFilterFunc))
 {
 }
 
-bool FAnalyticsSessionSummarySender::SendSessionSummary(const FString& UserId, const FString& AppId, const FString& AppVersion, const FString& SessionId, const TMap<FString, FString>& Properties)
+bool FAnalyticsSessionSummarySender::SendSessionSummary(const FString& UserId, const FString& AppId, const FString& AppVersion, const FString& SessionId, const TArray<FAnalyticsEventAttribute>& Properties)
 {
 	TArray<FAnalyticsEventAttribute> AnalyticsAttributes;
 
@@ -26,12 +26,12 @@ bool FAnalyticsSessionSummarySender::SendSessionSummary(const FString& UserId, c
 	AnalyticsAttributes.Emplace(TEXT("SummarySenderUserId"), AnalyticsProvider.GetUserID());
 	AnalyticsAttributes.Emplace(TEXT("SummarySenderSessionId"), AnalyticsProvider.GetSessionID()); // Not stripping the {} around the GUID like EditorSessionSummaryWriter does with SessionId.
 
-	for (const TPair<FString, FString>& Property : Properties)
+	for (const FAnalyticsEventAttribute& Property : Properties)
 	{
 		// Function is not bound or returns true.
-		if (!ShouldEmitPropFunc || ShouldEmitPropFunc(Property.Key, Property.Value))
+		if (!ShouldEmitPropFunc || ShouldEmitPropFunc(Property))
 		{
-			AnalyticsAttributes.Emplace(Property.Key, Property.Value);
+			AnalyticsAttributes.Emplace(Property);
 		}
 	}
 
diff --git a/Engine/Source/Runtime/Analytics/AnalyticsET/Private/Tests/AnalyticsSessionSummaryManagerTests.cpp b/Engine/Source/Runtime/Analytics/AnalyticsET/Private/Tests/AnalyticsSessionSummaryManagerTests.cpp
--- a/Engine/Source/Runtime/Analytics/AnalyticsET/Private/Tests/AnalyticsSessionSummaryManagerTests.cpp
+++ b/Engine/Source/Runtime/Analytics/AnalyticsET/Private/Tests/AnalyticsSessionSummaryManagerTests.cpp
@@ -18,7 +18,7 @@ IMPLEMENT_SIMPLE_AUTOMATION_TEST(FAnalyticsSessionSummaryManagerAutomationTest,
 class FDummyAnalyticsSummarySender : public IAnalyticsSessionSummarySender
 {
 public:
-	virtual bool SendSessionSummary(const FString& InUserId, const FString& InAppId, const FString& InAppVersion, const FString& InSessionId, const TMap<FString, FString>& InProperties) override
+	virtual bool SendSessionSummary(const FString& InUserId, const FString& InAppId, const FString& InAppVersion, const FString& InSessionId, const TArray<FAnalyticsEventAttribute>& InProperties) override
 	{
 		UserId = InUserId;
 		AppId = InAppId;
@@ -34,7 +34,7 @@ public:
 	FString AppId;
 	FString AppVersion;
 	FString SessionId;
-	TMap<FString, FString> Properties;
+	TArray<FAnalyticsEventAttribute> Properties;
 };
 
 bool FAnalyticsSessionSummaryManagerAutomationTest::RunTest(const FString& Parameters)
@@ -91,12 +91,12 @@ bool FAnalyticsSessionSummaryManagerAutomationTest::RunTest(const FString& Param
 		check(SummarySender->SessionId  == PrincipalProcessSessionId);
 
 		// Check the automatically embedded keys.
-		check(*SummarySender->Properties.Find(DelayedSendProp.Key)  == LexToString(false));
-		check(*SummarySender->Properties.Find(SessionIdProp.Key)    == PrincipalProcessSessionId);
-		check(*SummarySender->Properties.Find(SentFromProp.Key)     == PrincipalProcessName);
+		check(SummarySender->Properties.FindByPredicate([&](const FAnalyticsEventAttribute& Attr) { return Attr.GetName() == DelayedSendProp.Key; })->GetValue() == LexToString(false));
+		check(SummarySender->Properties.FindByPredicate([&](const FAnalyticsEventAttribute& Attr) { return Attr.GetName() == SessionIdProp.Key; })->GetValue()   == PrincipalProcessSessionId);
+		check(SummarySender->Properties.FindByPredicate([&](const FAnalyticsEventAttribute& Attr) { return Attr.GetName() == SentFromProp.Key; })->GetValue()    == PrincipalProcessName);
 
 		// Check the custom property.
-		check(*SummarySender->Properties.Find(Prop1.Key) == LexToString(Prop1Value));
+		check(SummarySender->Properties.FindByPredicate([&](const FAnalyticsEventAttribute& Attr) { return Attr.GetName() == Prop1.Key; })->GetValue() == LexToString(Prop1Value));
 	}
 
 	//
diff --git a/Engine/Source/Runtime/Analytics/AnalyticsET/Public/AnalyticsSessionSummaryManager.h b/Engine/Source/Runtime/Analytics/AnalyticsET/Public/AnalyticsSessionSummaryManager.h
--- a/Engine/Source/Runtime/Analytics/AnalyticsET/Public/AnalyticsSessionSummaryManager.h
+++ b/Engine/Source/Runtime/Analytics/AnalyticsET/Public/AnalyticsSessionSummaryManager.h
@@ -4,6 +4,7 @@
 
 #include "CoreMinimal.h"
 #include "Interfaces/IAnalyticsPropertyStore.h"
+#include "AnalyticsEventAttribute.h"
 
 class IAnalyticsSessionSummarySender;
 
@@ -148,7 +149,7 @@ private:
 	void ProcessSummary(const FString& ProcessGroupId, const FProcessGroup& ProcessGroup);
 
 	/** Loads the summaries created for the application and its subsidiary processes (like Editor/CRC combo) and aggregates summary properties. */
-	bool AggregateSummaries(const FString& ProcessGroupId, const TArray<FPropertyFileInfo>& PropertyFiles, TMap<FString, FString>& OutSummaryProperties, TMap<FString, FString>& OutInternalProperties);
+	bool AggregateSummaries(const FString& ProcessGroupId, const FProcessGroup& ProcessGroup, TArray<FAnalyticsEventAttribute>& OutSummaryProperties, TArray<FAnalyticsEventAttribute>& OutInternalProperties);
 
 	/** Deletes the files used to record and perist the properties. */
 	bool CleanupFiles(const TArray<FPropertyFileInfo>& PropertyFiles, bool bOnSuccess);
diff --git a/Engine/Source/Runtime/Analytics/AnalyticsET/Public/AnalyticsSessionSummarySender.h b/Engine/Source/Runtime/Analytics/AnalyticsET/Public/AnalyticsSessionSummarySender.h
--- a/Engine/Source/Runtime/Analytics/AnalyticsET/Public/AnalyticsSessionSummarySender.h
+++ b/Engine/Source/Runtime/Analytics/AnalyticsET/Public/AnalyticsSessionSummarySender.h
@@ -18,7 +18,7 @@ public:
 	 * @param Provider The analytics provider that will be used to emits the summary event.
 	 * @param ShouldEmitFilterFunc A filter function invoked for each properties before sending. If the function is bound and returns false, the property will be filtered out. If unbound, all properties passed to SendSesionSummary() are emitted.
 	 */
-	FAnalyticsSessionSummarySender(IAnalyticsProviderET& Provider, TFunction<bool(const FString&, const FString&)> ShouldEmitFilterFunc = TFunction<bool(const FString&, const FString&)>());
+	FAnalyticsSessionSummarySender(IAnalyticsProviderET& Provider, TFunction<bool(const FAnalyticsEventAttribute&)> ShouldEmitFilterFunc = TFunction<bool(const FAnalyticsEventAttribute&)>());
 
 	/**
 	 * Emits the summary events for the specified session id on behalf of the specified user/app/appversion. The function filters the properties by invoking the functor specified at construction.
@@ -28,9 +28,9 @@ public:
 	 * @param SessionId The session is for which the report is emitted.
 	 * @param Properties The list of properties that makes up the summary event.
 	 */
-	virtual bool SendSessionSummary(const FString& UserId, const FString& AppId, const FString& AppVersion, const FString& SessionId, const TMap<FString, FString>& Properties) override;
+	virtual bool SendSessionSummary(const FString& UserId, const FString& AppId, const FString& AppVersion, const FString& SessionId, const TArray<FAnalyticsEventAttribute>& Properties) override;
 
 private:
 	IAnalyticsProviderET& AnalyticsProvider;
-	TFunction<bool(const FString&, const FString&)> ShouldEmitPropFunc;
+	TFunction<bool(const FAnalyticsEventAttribute&)> ShouldEmitPropFunc;
 };
diff --git a/Engine/Source/Runtime/Analytics/AnalyticsET/Public/IAnalyticsSessionSummarySender.h b/Engine/Source/Runtime/Analytics/AnalyticsET/Public/IAnalyticsSessionSummarySender.h
--- a/Engine/Source/Runtime/Analytics/AnalyticsET/Public/IAnalyticsSessionSummarySender.h
+++ b/Engine/Source/Runtime/Analytics/AnalyticsET/Public/IAnalyticsSessionSummarySender.h
@@ -3,6 +3,7 @@
 #pragma once
 
 #include "CoreMinimal.h"
+#include "AnalyticsEventAttribute.h"
 
 /**
  * Sends the analytics session summary to a backend service.
@@ -20,5 +21,5 @@ public:
 	 * @param SessionId The session is for which the report is emitted.
 	 * @param Properties The list of properties that makes up the summary event.
 	 */
-	virtual bool SendSessionSummary(const FString& UserId, const FString& AppId, const FString& AppVersion, const FString& SessionId, const TMap<FString, FString>& Properties) = 0;
+	virtual bool SendSessionSummary(const FString& UserId, const FString& AppId, const FString& AppVersion, const FString& SessionId, const TArray<FAnalyticsEventAttribute>& Properties) = 0;
 };
