Fix non-unity CIS error on Android in AnalyticsSessionSummaryManager.cpp

diff --git a/Engine/Source/Runtime/Analytics/AnalyticsET/Private/AnalyticsSessionSummaryManager.cpp b/Engine/Source/Runtime/Analytics/AnalyticsET/Private/AnalyticsSessionSummaryManager.cpp
--- a/Engine/Source/Runtime/Analytics/AnalyticsET/Private/AnalyticsSessionSummaryManager.cpp
+++ b/Engine/Source/Runtime/Analytics/AnalyticsET/Private/AnalyticsSessionSummaryManager.cpp
@@ -10,6 +10,7 @@
 #include "HAL/PlatformTime.h"
 #include "Misc/App.h"
 #include "Misc/Optional.h"
+#include "Misc/Paths.h"
 #include "Internationalization/Regex.h"
 
 DEFINE_LOG_CATEGORY_STATIC(LogAnalyticsSessionSummary, Verbose, All);
