Fix CIS no unity errors

diff --git a/Engine/Source/Editor/AssetPlacementEdMode/Private/AssetPlacementEdModeStyle.cpp b/Engine/Source/Editor/AssetPlacementEdMode/Private/AssetPlacementEdModeStyle.cpp
--- a/Engine/Source/Editor/AssetPlacementEdMode/Private/AssetPlacementEdModeStyle.cpp
+++ b/Engine/Source/Editor/AssetPlacementEdMode/Private/AssetPlacementEdModeStyle.cpp
@@ -2,6 +2,7 @@
 
 #include "AssetPlacementEdModeStyle.h"
 
+#include "Styling/CoreStyle.h"
 #include "Styling/SlateStyle.h"
 #include "Styling/SlateStyleMacros.h"
 #include "Styling/SlateStyleRegistry.h"
diff --git a/Engine/Source/Editor/BlueprintGraph/Private/K2Node_ClassDynamicCast.cpp b/Engine/Source/Editor/BlueprintGraph/Private/K2Node_ClassDynamicCast.cpp
--- a/Engine/Source/Editor/BlueprintGraph/Private/K2Node_ClassDynamicCast.cpp
+++ b/Engine/Source/Editor/BlueprintGraph/Private/K2Node_ClassDynamicCast.cpp
@@ -5,6 +5,8 @@
 #include "GraphEditorSettings.h"
 #include "EdGraphSchema_K2.h"
 #include "DynamicCastHandler.h"
+#include "Engine/Blueprint.h"
+#include "UObject/Interface.h"
 
 #define LOCTEXT_NAMESPACE "K2Node_ClassDynamicCast"
 
diff --git a/Engine/Source/Runtime/Analytics/AnalyticsET/Private/Tests/AnalyticsSessionSummaryManagerTests.cpp b/Engine/Source/Runtime/Analytics/AnalyticsET/Private/Tests/AnalyticsSessionSummaryManagerTests.cpp
--- a/Engine/Source/Runtime/Analytics/AnalyticsET/Private/Tests/AnalyticsSessionSummaryManagerTests.cpp
+++ b/Engine/Source/Runtime/Analytics/AnalyticsET/Private/Tests/AnalyticsSessionSummaryManagerTests.cpp
@@ -2,6 +2,7 @@
 
 #include "CoreMinimal.h"
 #include "GenericPlatform/GenericPlatformFile.h"
+#include "HAL/FileManager.h"
 #include "HAL/PlatformFileManager.h"
 #include "HAL/PlatformProcess.h"
 #include "Interfaces/IAnalyticsPropertyStore.h"
diff --git a/Engine/Source/Runtime/Engine/Private/ReplayNetConnection.cpp b/Engine/Source/Runtime/Engine/Private/ReplayNetConnection.cpp
--- a/Engine/Source/Runtime/Engine/Private/ReplayNetConnection.cpp
+++ b/Engine/Source/Runtime/Engine/Private/ReplayNetConnection.cpp
@@ -4,6 +4,7 @@
 #include "Net/NetworkProfiler.h"
 #include "Net/NetworkGranularMemoryLogging.h"
 #include "Engine/LevelStreaming.h"
+#include "Engine/Engine.h"
 #include "Engine/World.h"
 #include "Engine/ActorChannel.h"
 #include "Engine/NetworkObjectList.h"
diff --git a/Engine/Source/Runtime/RHI/Private/RHI.cpp b/Engine/Source/Runtime/RHI/Private/RHI.cpp
--- a/Engine/Source/Runtime/RHI/Private/RHI.cpp
+++ b/Engine/Source/Runtime/RHI/Private/RHI.cpp
@@ -17,6 +17,10 @@
 #include "Misc/BufferedOutputDevice.h"
 #include "Misc/OutputDeviceFile.h"
 
+#if RHI_WANT_RESOURCE_INFO
+#include "HAL/FileManager.h"
+#endif
+
 IMPLEMENT_MODULE(FDefaultModuleImpl, RHI);
 
 /** RHI Logging. */
diff --git a/Engine/Source/Runtime/Renderer/Private/HairStrands/HairStrandsData.h b/Engine/Source/Runtime/Renderer/Private/HairStrands/HairStrandsData.h
--- a/Engine/Source/Runtime/Renderer/Private/HairStrands/HairStrandsData.h
+++ b/Engine/Source/Runtime/Renderer/Private/HairStrands/HairStrandsData.h
@@ -8,11 +8,11 @@
 
 #include "CoreMinimal.h"
 #include "RendererInterface.h"
+#include "RHIGPUReadback.h"
 #include "Shader.h"
 
 class FPrimitiveSceneProxy;
 class FViewInfo;
-class FRHIGPUBufferReadback;
 struct FMeshBatch;
 
 ////////////////////////////////////////////////////////////////////////////////////
