Fix FTCHARToUTF8_Convert::Convert() to take a pointer as the first parameter.

diff --git a/Engine/Source/Runtime/Analytics/AnalyticsET/Private/AnalyticsProviderETEventCache.cpp b/Engine/Source/Runtime/Analytics/AnalyticsET/Private/AnalyticsProviderETEventCache.cpp
--- a/Engine/Source/Runtime/Analytics/AnalyticsET/Private/AnalyticsProviderETEventCache.cpp
+++ b/Engine/Source/Runtime/Analytics/AnalyticsET/Private/AnalyticsProviderETEventCache.cpp
@@ -114,7 +114,7 @@ namespace EventCacheStatic
 			// make space for the string
 			UTF8Stream.SetNumUninitialized(OldLen + StrLen, false);
 			// convert it to UTF8
-			int32 CharsWritten = FTCHARToUTF8_Convert::Convert(&UTF8Stream[OldLen], StrLen, Str, Len);
+			int32 CharsWritten = FTCHARToUTF8_Convert::Convert((ANSICHAR*)&UTF8Stream[OldLen], StrLen, Str, Len);
 			// figure out how many characters were actually written 
 			if (CharsWritten >= 0)
 			{
@@ -132,7 +132,7 @@ namespace EventCacheStatic
 					const int32 ActualCharsNeeded = FTCHARToUTF8_Convert::ConvertedLength(Str, Len);
 					UTF8Stream.SetNumUninitialized(OldLen + ActualCharsNeeded, false);
 					// convert it to UTF8 using the known number of charts
-					FTCHARToUTF8_Convert::Convert(&UTF8Stream[OldLen], ActualCharsNeeded, Str, Len);
+					FTCHARToUTF8_Convert::Convert((ANSICHAR*)&UTF8Stream[OldLen], ActualCharsNeeded, Str, Len);
 					bWroteFullString = true;
 				}
 			}
diff --git a/Engine/Source/Runtime/Core/Public/Containers/StringConv.h b/Engine/Source/Runtime/Core/Public/Containers/StringConv.h
--- a/Engine/Source/Runtime/Core/Public/Containers/StringConv.h
+++ b/Engine/Source/Runtime/Core/Public/Containers/StringConv.h
@@ -285,11 +285,10 @@ public:
 	 * @return          The number of bytes written to Dest, up to DestLen, or -1 if the entire Source string could did not fit in DestLen bytes.
 	 */
 	template <
-		typename DestBufferType,
 		typename SrcBufferType,
 		std::enable_if_t<FPlatformString::IsCharEncodingCompatibleWith<SrcBufferType, FromType>()>* = nullptr
 	>
-	static FORCEINLINE int32 Convert(DestBufferType Dest, int32 DestLen, const SrcBufferType* Source, int32 SourceLen)
+	static FORCEINLINE int32 Convert(ToType* Dest, int32 DestLen, const SrcBufferType* Source, int32 SourceLen)
 	{
 		return (int32)(FPlatformString::Convert((IntendedToType*)Dest, DestLen, (const FromType*)Source, SourceLen) - (IntendedToType*)Dest);
 	}
