Fix case sensitive file name, also added missing #include for non pch builds

diff --git a/Engine/Source/Runtime/Analytics/Analytics/Private/AnalyticsPropertyStore.cpp b/Engine/Source/Runtime/Analytics/Analytics/Private/AnalyticsPropertyStore.cpp
--- a/Engine/Source/Runtime/Analytics/Analytics/Private/AnalyticsPropertyStore.cpp
+++ b/Engine/Source/Runtime/Analytics/Analytics/Private/AnalyticsPropertyStore.cpp
@@ -4,7 +4,7 @@
 #include "Containers/StringConv.h"
 #include "GenericPlatform/GenericPlatformFile.h"
 #include "HAL/FileManager.h"
-#include "HAL/PlatformFileManager.h"
+#include "HAL/PlatformFilemanager.h"
 #include "HAL/PlatformTime.h"
 #include "Misc/SecureHash.h"
 #include "Misc/ScopeExit.h"
diff --git a/Engine/Source/Runtime/Analytics/Analytics/Private/Tests/AnalyticsPropertyStoreTests.cpp b/Engine/Source/Runtime/Analytics/Analytics/Private/Tests/AnalyticsPropertyStoreTests.cpp
--- a/Engine/Source/Runtime/Analytics/Analytics/Private/Tests/AnalyticsPropertyStoreTests.cpp
+++ b/Engine/Source/Runtime/Analytics/Analytics/Private/Tests/AnalyticsPropertyStoreTests.cpp
@@ -3,7 +3,7 @@
 #include "CoreMinimal.h"
 #include "GenericPlatform/GenericPlatformFile.h"
 #include "HAL/FileManager.h"
-#include "HAL/PlatformFileManager.h"
+#include "HAL/PlatformFilemanager.h"
 #include "Misc/AutomationTest.h"
 #include "Misc/Paths.h"
 #include "AnalyticsPropertyStore.h"
diff --git a/Engine/Source/Runtime/Analytics/AnalyticsET/Private/Tests/AnalyticsSessionSummaryManagerTests.cpp b/Engine/Source/Runtime/Analytics/AnalyticsET/Private/Tests/AnalyticsSessionSummaryManagerTests.cpp
--- a/Engine/Source/Runtime/Analytics/AnalyticsET/Private/Tests/AnalyticsSessionSummaryManagerTests.cpp
+++ b/Engine/Source/Runtime/Analytics/AnalyticsET/Private/Tests/AnalyticsSessionSummaryManagerTests.cpp
@@ -3,7 +3,7 @@
 #include "CoreMinimal.h"
 #include "GenericPlatform/GenericPlatformFile.h"
 #include "HAL/FileManager.h"
-#include "HAL/PlatformFileManager.h"
+#include "HAL/PlatformFilemanager.h"
 #include "HAL/PlatformProcess.h"
 #include "Interfaces/IAnalyticsPropertyStore.h"
 #include "Misc/AutomationTest.h"
