Rremaining Direct FTCHARToUTF8_Convert usage replaced with StringCast or FPlatformString.

diff --git a/Engine/Source/Runtime/Analytics/AnalyticsET/Private/AnalyticsProviderETEventCache.cpp b/Engine/Source/Runtime/Analytics/AnalyticsET/Private/AnalyticsProviderETEventCache.cpp
--- a/Engine/Source/Runtime/Analytics/AnalyticsET/Private/AnalyticsProviderETEventCache.cpp
+++ b/Engine/Source/Runtime/Analytics/AnalyticsET/Private/AnalyticsProviderETEventCache.cpp
@@ -68,7 +68,6 @@ namespace EventCacheStatic
 	// This is the buffer we will convert strings into UTF8 into, since it's difficult to convert them directly into a TArray<>, since it doesn't know how to resize itself.
 	// We also don't want to walk the string once to count the chars if we don't have to. so we pay the price to copy directly into a stack-allocated buffer most of the time,
 	// but let it spill over to a dynamic allocation for long strings.
-	//typedef TStringConversion<FTCHARToUTF8_Convert, ConversionBufferSize> FPayloadUTF8Converter;
 	typedef TStringBuilder<ConversionBufferSize> FJsonStringBuilder;
 
 	const ANSICHAR* PayloadTemplate = "{\"Events\":[]}";
@@ -98,9 +97,9 @@ namespace EventCacheStatic
 		// *** ORIGINAL, simpler code. But slower. ***
 		// convert directly into new array, precompute length
 		// get the string length and expand our buffer to fit it.
-		//const int32 StrLen = FTCHARToUTF8_Convert::ConvertedLength(Str, Len);
+		//const int32 StrLen = FPlatformString::ConvertedLength<UTF8CHAR>(Str, Len);
 		//UTF8Stream.SetNumUninitialized(OldLen + StrLen, false);
-		//FTCHARToUTF8_Convert::Convert(&UTF8Stream[OldLen], StrLen, Str, Len);
+		//FPlatformString::Convert((UTF8CHAR*)&UTF8Stream[OldLen], StrLen, Str, Len);
 
 		// optimistically allocate a bit of extra space and see if we fill up the buffer.
 		// If we do, lengthen the buffer a bit and try again.
@@ -114,12 +113,10 @@ namespace EventCacheStatic
 			// make space for the string
 			UTF8Stream.SetNumUninitialized(OldLen + StrLen, false);
 			// convert it to UTF8
-			int32 CharsWritten = FTCHARToUTF8_Convert::Convert((UTF8CHAR*)&UTF8Stream[OldLen], StrLen, Str, Len);
-			// figure out how many characters were actually written 
-			if (CharsWritten >= 0)
+			if (UTF8CHAR* NewEnd = FPlatformString::Convert((UTF8CHAR*)&UTF8Stream[OldLen], StrLen, Str, Len))
 			{
 				// truncate to that length.
-				UTF8Stream.SetNum(OldLen + CharsWritten, false);
+				UTF8Stream.SetNum(OldLen + (int32)(NewEnd - (UTF8CHAR*)&UTF8Stream[OldLen]), false);
 				bWroteFullString = true;
 			}
 			else
@@ -129,10 +126,10 @@ namespace EventCacheStatic
 				// if we grow too much, give up and compute the true chars needed.
 				if (SizeMultiplier >= 2.0)
 				{
-					const int32 ActualCharsNeeded = FTCHARToUTF8_Convert::ConvertedLength(Str, Len);
+					const int32 ActualCharsNeeded = FPlatformString::ConvertedLength<UTF8CHAR>(Str, Len);
 					UTF8Stream.SetNumUninitialized(OldLen + ActualCharsNeeded, false);
 					// convert it to UTF8 using the known number of charts
-					FTCHARToUTF8_Convert::Convert((UTF8CHAR*)&UTF8Stream[OldLen], ActualCharsNeeded, Str, Len);
+					FPlatformString::Convert((UTF8CHAR*)&UTF8Stream[OldLen], ActualCharsNeeded, Str, Len);
 					bWroteFullString = true;
 				}
 			}
diff --git a/Engine/Source/Runtime/Online/HTTP/Private/GenericPlatform/GenericPlatformHttp.cpp b/Engine/Source/Runtime/Online/HTTP/Private/GenericPlatform/GenericPlatformHttp.cpp
--- a/Engine/Source/Runtime/Online/HTTP/Private/GenericPlatform/GenericPlatformHttp.cpp
+++ b/Engine/Source/Runtime/Online/HTTP/Private/GenericPlatform/GenericPlatformHttp.cpp
@@ -181,7 +181,7 @@ FString FGenericPlatformHttp::UrlDecode(const FStringView EncodedString)
 	FTCHARToUTF8 Converter(EncodedString.GetData(), EncodedString.Len());
 	const UTF8CHAR* UTF8Data = (UTF8CHAR*)Converter.Get();	
 	
-	TArray<ANSICHAR> Data;
+	TArray<UTF8CHAR> Data;
 	Data.Reserve(EncodedString.Len());
 
 	for (int32 CharIdx = 0; CharIdx < Converter.Length();)
@@ -200,12 +200,13 @@ FString FGenericPlatformHttp::UrlDecode(const FStringView EncodedString)
 					Value += FParse::HexDigit(UTF8Data[CharIdx + 5]);
 					CharIdx += 6;
 
-					ANSICHAR Buffer[8] = { 0 };
-					ANSICHAR* BufferPtr = Buffer;
+					UTF8CHAR Buffer[8] = {};
+					UTF8CHAR* BufferPtr = Buffer;
 					const int32 Len = UE_ARRAY_COUNT(Buffer);
-					const int32 WrittenChars = FTCHARToUTF8_Convert::Utf8FromCodepoint(Value, BufferPtr, Len);
+					const UTF8CHAR* BufferEnd = FPlatformString::Convert(BufferPtr, Len, (UTF32CHAR*)&Value, 1);
+					check(BufferEnd);
 
-					Data.Append(Buffer, WrittenChars);
+					Data.Append(Buffer, (int32)(BufferEnd - BufferPtr));
 				}
 				else
 				{
@@ -220,7 +221,7 @@ FString FGenericPlatformHttp::UrlDecode(const FStringView EncodedString)
 				Value = FParse::HexDigit(UTF8Data[CharIdx + 1]) << 4;
 				Value += FParse::HexDigit(UTF8Data[CharIdx + 2]);
 				CharIdx += 3;
-				Data.Add((ANSICHAR)(Value));
+				Data.Add((UTF8CHAR)(Value));
 			}
 			else
 			{
@@ -237,7 +238,7 @@ FString FGenericPlatformHttp::UrlDecode(const FStringView EncodedString)
 		}
 	}
 
-	Data.Add('\0');
+	Data.Add(UTF8TEXT('\0'));
 	return FString(UTF8_TO_TCHAR(Data.GetData()));
 }
 
diff --git a/Engine/Source/Runtime/Online/HTTPServer/Private/HttpConnectionRequestReadContext.cpp b/Engine/Source/Runtime/Online/HTTPServer/Private/HttpConnectionRequestReadContext.cpp
--- a/Engine/Source/Runtime/Online/HTTPServer/Private/HttpConnectionRequestReadContext.cpp
+++ b/Engine/Source/Runtime/Online/HTTPServer/Private/HttpConnectionRequestReadContext.cpp
@@ -324,7 +324,7 @@ FString FHttpConnectionRequestReadContext::UrlDecode(const FString &EncodedStrin
 	FTCHARToUTF8 Converter(*EncodedString);
 	const UTF8CHAR* UTF8Data = (UTF8CHAR*)Converter.Get();
 
-	TArray<ANSICHAR> Data;
+	TArray<UTF8CHAR> Data;
 	Data.Reserve(EncodedString.Len());
 
 	for (int32 CharIdx = 0; CharIdx < Converter.Length();)
@@ -343,12 +343,13 @@ FString FHttpConnectionRequestReadContext::UrlDecode(const FString &EncodedStrin
 					Value += FParse::HexDigit(UTF8Data[CharIdx + 5]);
 					CharIdx += 6;
 
-					ANSICHAR Buffer[8] = { 0 };
-					ANSICHAR* BufferPtr = Buffer;
+					UTF8CHAR Buffer[8] = { };
+					UTF8CHAR* BufferPtr = Buffer;
 					const int32 Len = UE_ARRAY_COUNT(Buffer);
-					const int32 WrittenChars = FTCHARToUTF8_Convert::Utf8FromCodepoint(Value, BufferPtr, Len);
+					const UTF8CHAR* BufferEnd = FPlatformString::Convert(BufferPtr, Len, (UTF32CHAR*)&Value, 1);
+					check(BufferEnd);
 
-					Data.Append(Buffer, WrittenChars);
+					Data.Append(Buffer, (int32)(BufferEnd - BufferPtr));
 				}
 				else
 				{
@@ -363,7 +364,7 @@ FString FHttpConnectionRequestReadContext::UrlDecode(const FString &EncodedStrin
 				Value = FParse::HexDigit(UTF8Data[CharIdx + 1]) << 4;
 				Value += FParse::HexDigit(UTF8Data[CharIdx + 2]);
 				CharIdx += 3;
-				Data.Add((ANSICHAR)(Value));
+				Data.Add((UTF8CHAR)(Value));
 			}
 			else
 			{
@@ -380,6 +381,6 @@ FString FHttpConnectionRequestReadContext::UrlDecode(const FString &EncodedStrin
 		}
 	}
 
-	Data.Add('\0');
+	Data.Add(UTF8TEXT('\0'));
 	return FString(UTF8_TO_TCHAR(Data.GetData()));
 }
\ No newline at end of file
