Fixed Analytics asserting and crashing the Editor when the analytics property store files contains empty key names, likely from some sort of corruption in the file header.

diff --git a/Engine/Source/Runtime/Analytics/Analytics/Private/AnalyticsPropertyStore.cpp b/Engine/Source/Runtime/Analytics/Analytics/Private/AnalyticsPropertyStore.cpp
--- a/Engine/Source/Runtime/Analytics/Analytics/Private/AnalyticsPropertyStore.cpp
+++ b/Engine/Source/Runtime/Analytics/Analytics/Private/AnalyticsPropertyStore.cpp
@@ -137,36 +137,34 @@ bool FAnalyticsPropertyStore::Load(const FString& Pathname)
 	Reset();
 
 	// Open the file in read/write mode (random access needed).
-	FileHandle = TUniquePtr<IFileHandle>(FPlatformFileManager::Get().GetPlatformFile().OpenWrite(*Pathname, /*bAppend*/true, /*bAllowRead*/true));
-	if (!FileHandle)
+	TUniquePtr<IFileHandle> ScopedFileHandle(FPlatformFileManager::Get().GetPlatformFile().OpenWrite(*Pathname, /*bAppend*/true, /*bAllowRead*/true));
+	if (!ScopedFileHandle)
 	{
 		return false;
 	}
 
 	// Ensure the file header data is there. [Version | Checksum | StoreDataSize]. The file can be larger that the stored data size.
-	int64 FileSize = FileHandle->Size();
+	int64 FileSize = ScopedFileHandle->Size();
 	if (FileSize < sizeof(AnalyticsPropertyStoreUtils::FAnalyticsStoreFileHeader))
 	{
-		FileHandle.Reset();
 		return false;
 	}
 
 	// Read the header.
-	FileHandle->Seek(0);
+	ScopedFileHandle->Seek(0);
 	AnalyticsPropertyStoreUtils::FAnalyticsStoreFileHeader FileHeader;
-	FileHandle->Read(reinterpret_cast<uint8*>(&FileHeader), sizeof(FileHeader));
+	ScopedFileHandle->Read(reinterpret_cast<uint8*>(&FileHeader), sizeof(FileHeader));
 
 	// If the file format is incompatible.
 	if (FileHeader.Version != AnalyticsPropertyStoreUtils::StoreFormatVersion)
 	{
-		FileHandle.Reset();
 		return false;
 	}
 
 	// Read the stored data.
 	StorageBuf.Empty(FileHeader.DataSize);
 	StorageBuf.AddUninitialized(FileHeader.DataSize);
-	FileHandle->Read(StorageBuf.GetData(), FileHeader.DataSize);
+	ScopedFileHandle->Read(StorageBuf.GetData(), FileHeader.DataSize);
 
 	// Compute the checksum from the stored data.
 	uint8 Checksum[AnalyticsPropertyStoreUtils::Md5DigestSize];
@@ -177,7 +175,6 @@ bool FAnalyticsPropertyStore::Load(const FString& Pathname)
 	// If the checksums don't match.
 	if (FMemory::Memcmp(FileHeader.Checksum, Checksum, AnalyticsPropertyStoreUtils::Md5DigestSize) != 0)
 	{
-		FileHandle.Reset();
 		return false;
 	}
 
@@ -202,7 +199,11 @@ bool FAnalyticsPropertyStore::Load(const FString& Pathname)
 				static_assert(sizeof(uint32) == sizeof(int32) && sizeof(uint32) == sizeof(float), "Incompatible data size");
 				StorageReader.Seek(StorageReader.Tell() + sizeof(uint32)); // Skip over the value, it is not needed.
 				StorageReader << Key;
-				if (!IsDead(TypeCode))
+				if (Key.IsEmpty())
+				{
+					return false; // Invalid and not expected. Might be because the header or payload got corrupted before/during the write, likely a race condition in FAnalyticsPropertyStore::Flush()
+				}
+				else if (!IsDead(TypeCode))
 				{
 					NameOffsetMap.Emplace(Key, Offset);
 				}
@@ -217,7 +218,11 @@ bool FAnalyticsPropertyStore::Load(const FString& Pathname)
 				static_assert(sizeof(uint64) == sizeof(int64) && sizeof(uint64) == sizeof(double) && sizeof(uint64) == sizeof(decltype(static_cast<FDateTime*>(nullptr)->GetTicks())), "Incompatible data size");
 				StorageReader.Seek(StorageReader.Tell() + sizeof(uint64)); // Skip over the value, it is not needed.
 				StorageReader << Key;
-				if (!IsDead(TypeCode))
+				if (Key.IsEmpty())
+				{
+					return false; // Invalid and not expected. Might be because the header or payload got corrupted before/during the write, likely a race condition in FAnalyticsPropertyStore::Flush()
+				}
+				else if (!IsDead(TypeCode))
 				{
 					NameOffsetMap.Emplace(Key, Offset);
 				}
@@ -228,7 +233,11 @@ bool FAnalyticsPropertyStore::Load(const FString& Pathname)
 			{
 				StorageReader.Seek(StorageReader.Tell() + sizeof(bool)); // Skip over the value, it is not needed.
 				StorageReader << Key;
-				if (!IsDead(TypeCode))
+				if (Key.IsEmpty())
+				{
+					return false; // Invalid and not expected. Might be because the header or payload got corrupted before/during the write, likely a race condition in FAnalyticsPropertyStore::Flush()
+				}
+				else if (!IsDead(TypeCode))
 				{
 					NameOffsetMap.Emplace(Key, Offset);
 				}
@@ -242,7 +251,11 @@ bool FAnalyticsPropertyStore::Load(const FString& Pathname)
 				StorageReader.Serialize(&CapacityInBytes, sizeof(CapacityInBytes)); // Read the capacity.
 				StorageReader.Seek(StorageReader.Tell() + CapacityInBytes); // Skip over the value and its extra capacity space if any.
 				StorageReader << Key;
-				if (!IsDead(TypeCode))
+				if (Key.IsEmpty())
+				{
+					return false; // Invalid and not expected. Might be because the header or payload got corrupted before/during the write, likely a race condition in FAnalyticsPropertyStore::Flush()
+				}
+				else if (!IsDead(TypeCode))
 				{
 					NameOffsetMap.Emplace(Key, Offset);
 				}
@@ -254,6 +267,8 @@ bool FAnalyticsPropertyStore::Load(const FString& Pathname)
 		}
 	}
 
+	// Keep the file handle/open since the file content is valid.
+	FileHandle = MoveTemp(ScopedFileHandle);
 	return true;
 }
 
@@ -944,7 +959,7 @@ bool FAnalyticsPropertyStore::Flush(bool bAsync, const FTimespan& Timeout)
 		return true; // Flushed.
 	};
 
-	// If the caller doesn't want to wait to get its request started.
+	// If the caller doesn't want to wait for a previous execution (if any) to complete.
 	if (Timeout.IsZero() && StoreLock.TryLock())
 	{
 		ON_SCOPE_EXIT { StoreLock.Unlock(); };
