Allow additional Epic-created session summaries to be generated by editor plugins

diff --git a/Engine/Source/Runtime/Analytics/AnalyticsET/Private/AnalyticsSessionSummaryManager.cpp b/Engine/Source/Runtime/Analytics/AnalyticsET/Private/AnalyticsSessionSummaryManager.cpp
--- a/Engine/Source/Runtime/Analytics/AnalyticsET/Private/AnalyticsSessionSummaryManager.cpp
+++ b/Engine/Source/Runtime/Analytics/AnalyticsET/Private/AnalyticsSessionSummaryManager.cpp
@@ -65,8 +65,8 @@ namespace AnalyticsSessionUtils
 
 	FRegexPattern GetAnalyticsPropertyFilenamePattern()
 	{
-		// Filename format is like "ProcessGroupGUID_PID_PID_ProcessName", for example 3E1D7ADBC38F4A789939D781E1B91520_12345_232513_Editor"
-		return FRegexPattern(TEXT(R"((^[0-9A-Za-z]+)_([0-9]+)_([0-9]+)_(.*))")); // Need help with regex? Try https://regex101.com/
+		// Filename format is like "ProcessGroupGUID_PID_PID_StoreCount_ProcessName", for example 3E1D7ADBC38F4A789939D781E1B91520_12345_232513_0_Editor"
+		return FRegexPattern(TEXT(R"((^[0-9A-Za-z]+)_([0-9]+)_([0-9]+)(_[0-9]+)?_(.*))")); // Need help with regex? Try https://regex101.com/
 	}
 
 	/** Try to acquire the rigth to process and send sessions for orphan groups. */
@@ -190,6 +190,7 @@ FAnalyticsSessionSummaryManager::FAnalyticsSessionSummaryManager(const FString&
 	, SessionRootPath(InSavedDir)
 	, CurrentProcessId(InCurrentProcessId)
 	, PrincipalProcessId(InPrincipalProcessId)
+	, StoreCounter(0)
 	, NextOrphanSessionCheckTimeSecs(FPlatformTime::Seconds())
 	, bOrphanGroupOwner(false)
 	, bIsPrincipal(InCurrentProcessId == InPrincipalProcessId)
@@ -220,13 +221,13 @@ FAnalyticsSessionSummaryManager::~FAnalyticsSessionSummaryManager()
 TSharedPtr<IAnalyticsPropertyStore> FAnalyticsSessionSummaryManager::MakeStore(uint32 InitialCapacity)
 {
 	// NOTE: The ProcessGroupId (a GUID) + PrincipalProcessId is expected to provide a unique key that should not collide with existing files.
-	FString PropertyStorePathname = SessionRootPath / FString::Printf(TEXT("%s_%d_%d_%s"), *ProcessGroupId, PrincipalProcessId, CurrentProcessId, *ProcessName);
+	FString PropertyStorePathname = SessionRootPath / FString::Printf(TEXT("%s_%d_%d_%d_%s"), *ProcessGroupId, PrincipalProcessId, CurrentProcessId, StoreCounter, *ProcessName);
 
 	TSharedPtr<FAnalyticsPropertyStore> PropertyStore = MakeShared<FAnalyticsPropertyStore>();
 	if (PropertyStore && PropertyStore->Create(PropertyStorePathname, InitialCapacity))
 	{
 		// If this is the principal process for which the session summary is created.
-		if (IsPrincipalProcess())
+		if (IsPrincipalProcess() && StoreCounter == 0)
 		{
 			// Add some internal key required to enable a subsidiary process to sent the summary on the behalf of this process.
 			AnalyticsManagerProperties::InternalSessionUserId.Set(PropertyStore.Get(), UserId);
@@ -234,6 +235,8 @@ TSharedPtr<IAnalyticsPropertyStore> FAnalyticsSessionSummaryManager::MakeStore(u
 			AnalyticsManagerProperties::InternalSessionAppVersion.Set(PropertyStore.Get(), AppVersion);
 			AnalyticsManagerProperties::InternalSessionId.Set(PropertyStore.Get(), SessionId);
 		}
+
+		++StoreCounter;
 	}
 	else
 	{
@@ -274,6 +277,8 @@ void FAnalyticsSessionSummaryManager::Shutdown(bool bDiscard)
 			CleanupFiles(ProcessGroup->PropertyFiles, /*bOnSuccess*/false);
 		}
 	}
+
+	StoreCounter = 0;
 }
 
 void FAnalyticsSessionSummaryManager::Tick()
@@ -317,7 +322,7 @@ TMap<FString, FAnalyticsSessionSummaryManager::FProcessGroup> FAnalyticsSessionS
 	{
 		if (!bIsDirectory)
 		{
-			// The property filename is "ProcessGroupGUID_PrincipalProcessId_WriterProcessId_WriterProcessName". The ProcessGroupGUID is expected to have EGuidFormats::Digits -> "00000A00C0000F000000000000000000"
+			// The property filename is "ProcessGroupGUID_PrincipalProcessId_WriterProcessId_StoreCount_WriterProcessName". The ProcessGroupGUID is expected to have EGuidFormats::Digits -> "00000A00C0000F000000000000000000"
 			FRegexMatcher Matcher(FilenamePattern, FPaths::GetCleanFilename(Pathname));
 			if (Matcher.FindNext())
 			{
@@ -329,7 +334,7 @@ TMap<FString, FAnalyticsSessionSummaryManager::FProcessGroup> FAnalyticsSessionS
 
 				FPropertyFileInfo PropertyFileInfo;
 				PropertyFileInfo.ProcessId = static_cast<uint32>(FCString::Atoi64(*Matcher.GetCaptureGroup(3)));
-				PropertyFileInfo.ProcessName = Matcher.GetCaptureGroup(4);
+				PropertyFileInfo.ProcessName = Matcher.GetCaptureGroup(5);
 				PropertyFileInfo.Pathname = Pathname;
 
 				uint32 FilePrincipalProcessId = static_cast<uint32>(FCString::Atoi64(*Matcher.GetCaptureGroup(2)));
diff --git a/Engine/Source/Runtime/Analytics/AnalyticsET/Public/AnalyticsSessionSummaryManager.h b/Engine/Source/Runtime/Analytics/AnalyticsET/Public/AnalyticsSessionSummaryManager.h
--- a/Engine/Source/Runtime/Analytics/AnalyticsET/Public/AnalyticsSessionSummaryManager.h
+++ b/Engine/Source/Runtime/Analytics/AnalyticsET/Public/AnalyticsSessionSummaryManager.h
@@ -91,7 +91,7 @@ public:
 	/**
 	 * Creates a new property store associated to this manager process group. The manager only sends an analytics session summary when all
 	 * collecting processes have closed their property store. Ensure to flush and release the store before calling the manager Shutdown() to
-	 * ensure the session is sent as soon a possible.
+	 * ensure the session is sent as soon a possible. Subsequent calls to MakeStore will increment an internal counter in the store's filename.
 	 * @param InitialCapacity The amount of space to reserve in the file.
 	 */
 	TSharedPtr<IAnalyticsPropertyStore> MakeStore(uint32 InitialCapacity);
@@ -187,6 +187,7 @@ private:
 	FString SessionRootPath;
 	uint32 CurrentProcessId;
 	uint32 PrincipalProcessId;
+	uint32 StoreCounter;
 	double NextOrphanSessionCheckTimeSecs;
 	bool bOrphanGroupOwner = false;
 	bool bIsPrincipal = false;
diff --git a/Engine/Source/Runtime/Engine/Private/EngineAnalytics.cpp b/Engine/Source/Runtime/Engine/Private/EngineAnalytics.cpp
--- a/Engine/Source/Runtime/Engine/Private/EngineAnalytics.cpp
+++ b/Engine/Source/Runtime/Engine/Private/EngineAnalytics.cpp
@@ -35,6 +35,8 @@ TSharedPtr<FEngineSessionManager> FEngineAnalytics::SessionManager;
 static TUniquePtr<FAnalyticsSessionSummaryManager> AnalyticsSessionSummaryManager;
 static TUniquePtr<FEditorAnalyticsSessionSummary> EditorAnalyticSessionSummary;
 static TSharedPtr<FAnalyticsSessionSummarySender> AnalyticsSessionSummarySender;
+FSimpleMulticastDelegate FEngineAnalytics::OnInitializeEngineAnalytics;
+FSimpleMulticastDelegate FEngineAnalytics::OnShutdownEngineAnalytics;
 #endif
 
 static TSharedPtr<IAnalyticsProviderET> CreateEpicAnalyticsProvider()
@@ -72,6 +74,15 @@ IAnalyticsProviderET& FEngineAnalytics::GetProvider()
 	return *Analytics.Get();
 }
 
+#if WITH_EDITOR
+FAnalyticsSessionSummaryManager& FEngineAnalytics::GetSummaryManager()
+{
+	checkf(bIsInitialized && AnalyticsSessionSummaryManager.IsValid(), TEXT("FEngineAnalytics::GetSessionManager called outside of Initialize/Shutdown."));
+
+	return *AnalyticsSessionSummaryManager.Get();
+}
+#endif
+
 void FEngineAnalytics::Initialize()
 {
 	checkf(!bIsInitialized, TEXT("FEngineAnalytics::Initialize called more than once."));
@@ -171,6 +182,8 @@ void FEngineAnalytics::Initialize()
 				// Create the object responsible to collect the Editor session properties.
 				EditorAnalyticSessionSummary = MakeUnique<FEditorAnalyticsSessionSummary>(EditorPropertyStore, FGenericCrashContext::GetOutOfProcessCrashReporterProcessId());
 			}
+
+			OnInitializeEngineAnalytics.Broadcast();
 		}
 #endif
 	}
@@ -186,6 +199,8 @@ void FEngineAnalytics::Shutdown(bool bIsEngineShutdown)
 	}
 
 #if WITH_EDITOR
+	OnShutdownEngineAnalytics.Broadcast();
+
 	if (EditorAnalyticSessionSummary)
 	{
 		EditorAnalyticSessionSummary->Shutdown();
diff --git a/Engine/Source/Runtime/Engine/Public/EngineAnalytics.h b/Engine/Source/Runtime/Engine/Public/EngineAnalytics.h
--- a/Engine/Source/Runtime/Engine/Public/EngineAnalytics.h
+++ b/Engine/Source/Runtime/Engine/Public/EngineAnalytics.h
@@ -6,6 +6,7 @@
 #include "IAnalyticsProviderET.h" // NOTE: Consider changing the code to replace IAnalyticProvider.h by IAnalyticProviderET.h
 
 class FEngineSessionManager;
+class FAnalyticsSessionSummaryManager;
 class IAnalyticsProvider;
 class IAnalyticsProviderET;
 struct FAnalyticsEventAttribute;
@@ -29,14 +30,27 @@ public:
 	 */
 	static ENGINE_API IAnalyticsProviderET& GetProvider();
 
+#if WITH_EDITOR
+	/** This is intended ONLY for use by Epic Games so that plugins can append to the session summary. */
+	static ENGINE_API FAnalyticsSessionSummaryManager& GetSummaryManager();
+#endif
+
 	/** Helper function to determine if the provider is valid. */
 	static ENGINE_API bool IsAvailable() { return Analytics.IsValid(); }
 
 	/** Called to initialize the singleton. */
 	static ENGINE_API void Initialize();
+#if WITH_EDITOR
+	/** A hook for analytics initialization after the SummaryManager is initialized */
+	static ENGINE_API FSimpleMulticastDelegate OnInitializeEngineAnalytics;
+#endif
 
 	/** Called to shut down the singleton */
 	static ENGINE_API void Shutdown(bool bIsEngineShutdown = false);
+#if WITH_EDITOR
+	/** A hook for analytics shutdown before the SummaryManager is shutdown */
+	static ENGINE_API FSimpleMulticastDelegate OnShutdownEngineAnalytics;
+#endif
 
 	static ENGINE_API void Tick(float DeltaTime);
 
