Fix several check() that should be verify() must still run in shipping

diff --git a/Engine/Plugins/Animation/ControlRig/Source/ControlRigDeveloper/Private/ControlRigBlueprint.cpp b/Engine/Plugins/Animation/ControlRig/Source/ControlRigDeveloper/Private/ControlRigBlueprint.cpp
--- a/Engine/Plugins/Animation/ControlRig/Source/ControlRigDeveloper/Private/ControlRigBlueprint.cpp
+++ b/Engine/Plugins/Animation/ControlRig/Source/ControlRigDeveloper/Private/ControlRigBlueprint.cpp
@@ -3262,7 +3262,7 @@ void UControlRigBlueprint::RebuildGraphFromModel()
 	DECLARE_SCOPE_HIERARCHICAL_COUNTER_FUNC()
 
 	TGuardValue<bool> SelfGuard(bSuspendModelNotificationsForSelf, true);
-	check(GetOrCreateController());
+	verify(GetOrCreateController());
 
 	TArray<UEdGraph*> EdGraphs;
 	GetAllGraphs(EdGraphs);
diff --git a/Engine/Plugins/Animation/ControlRig/Source/ControlRigEditor/Private/Editor/ControlRigEditor.cpp b/Engine/Plugins/Animation/ControlRig/Source/ControlRigEditor/Private/Editor/ControlRigEditor.cpp
--- a/Engine/Plugins/Animation/ControlRig/Source/ControlRigEditor/Private/Editor/ControlRigEditor.cpp
+++ b/Engine/Plugins/Animation/ControlRig/Source/ControlRigEditor/Private/Editor/ControlRigEditor.cpp
@@ -4511,7 +4511,7 @@ void FControlRigEditor::OnWrappedPropertyChangedChainEvent(UDetailsViewWrapperOb
 					}
 					else
 					{
-						check(PropertyPath.RemoveFromStart(CurrentString));
+						verify(PropertyPath.RemoveFromStart(CurrentString));
 						TransformType = ERigTransformType::MakeCurrent(TransformType);
 					}
 
@@ -4521,7 +4521,7 @@ void FControlRigEditor::OnWrappedPropertyChangedChainEvent(UDetailsViewWrapperOb
 					}
 					else
 					{
-						check(PropertyPath.RemoveFromStart(LocalString));
+						verify(PropertyPath.RemoveFromStart(LocalString));
 						TransformType = ERigTransformType::MakeLocal(TransformType);
 					}
 
diff --git a/Engine/Plugins/Enterprise/DataprepEditor/Source/DataprepEditor/Private/DataprepSnapshot.cpp b/Engine/Plugins/Enterprise/DataprepEditor/Source/DataprepEditor/Private/DataprepSnapshot.cpp
--- a/Engine/Plugins/Enterprise/DataprepEditor/Source/DataprepEditor/Private/DataprepSnapshot.cpp
+++ b/Engine/Plugins/Enterprise/DataprepEditor/Source/DataprepEditor/Private/DataprepSnapshot.cpp
@@ -865,7 +865,7 @@ void FDataprepEditor::RestoreFromSnapshot(bool bUpdateViewport)
 
 		// Load the text file to a string
 		FString FileBuffer;
-		check( FFileHelper::LoadFileToString(FileBuffer, *PackageFilePath) );
+		verify( FFileHelper::LoadFileToString(FileBuffer, *PackageFilePath) );
 
 		// Set the GWorld to the preview world since ULevelFactory::FactoryCreateText uses GWorld
 		UWorld* PrevGWorld = GWorld;
diff --git a/Engine/Plugins/Experimental/GPULightmass/Source/GPULightmass/Private/Scene/Scene.cpp b/Engine/Plugins/Experimental/GPULightmass/Source/GPULightmass/Private/Scene/Scene.cpp
--- a/Engine/Plugins/Experimental/GPULightmass/Source/GPULightmass/Private/Scene/Scene.cpp
+++ b/Engine/Plugins/Experimental/GPULightmass/Source/GPULightmass/Private/Scene/Scene.cpp
@@ -2304,7 +2304,7 @@ void FScene::ApplyFinishedLightmapsToWorld()
 									{
 										// Non-zero LODs derive their MapBuildDataId from LOD0. In this case also regenerate LOD0 GUID
 										StaticMeshComponent->LODData[0].MapBuildDataId.Invalidate();
-										check(StaticMeshComponent->LODData[0].CreateMapBuildDataId(0));
+										verify(StaticMeshComponent->LODData[0].CreateMapBuildDataId(0));
 									}
 								}
 
diff --git a/Engine/Plugins/Experimental/VirtualScouting/Source/VirtualScoutingOpenXR/Private/VirtualScoutingOpenXR.cpp b/Engine/Plugins/Experimental/VirtualScouting/Source/VirtualScoutingOpenXR/Private/VirtualScoutingOpenXR.cpp
--- a/Engine/Plugins/Experimental/VirtualScouting/Source/VirtualScoutingOpenXR/Private/VirtualScoutingOpenXR.cpp
+++ b/Engine/Plugins/Experimental/VirtualScouting/Source/VirtualScoutingOpenXR/Private/VirtualScoutingOpenXR.cpp
@@ -199,8 +199,8 @@ void FVirtualScoutingOpenXRExtension::AddActions(XrInstance InInstance, FCreateA
 	}
 
 	XrPath LeftHand, RightHand;
-	check(XR_SUCCEEDED(xrStringToPath(Instance, "/user/hand/left", &LeftHand)));
-	check(XR_SUCCEEDED(xrStringToPath(Instance, "/user/hand/right", &RightHand)));
+	verify(XR_SUCCEEDED(xrStringToPath(Instance, "/user/hand/left", &LeftHand)));
+	verify(XR_SUCCEEDED(xrStringToPath(Instance, "/user/hand/right", &RightHand)));
 
 	FActionParams ActionParams;
 	ActionParams.Set = ActionSet;
diff --git a/Engine/Source/Runtime/Analytics/Analytics/Private/Tests/AnalyticsPropertyStoreTests.cpp b/Engine/Source/Runtime/Analytics/Analytics/Private/Tests/AnalyticsPropertyStoreTests.cpp
--- a/Engine/Source/Runtime/Analytics/Analytics/Private/Tests/AnalyticsPropertyStoreTests.cpp
+++ b/Engine/Source/Runtime/Analytics/Analytics/Private/Tests/AnalyticsPropertyStoreTests.cpp
@@ -63,11 +63,11 @@ bool FAnalyticsPropertyStoreAutomationTest::RunTest(const FString& Parameters)
 	// Create an empty store and reload it.
 	{
 		FAnalyticsPropertyStore Store;
-		check(Store.Create(TestStorePathname, /*CapacityHint*/2 * 1024));
+		verify(Store.Create(TestStorePathname, /*CapacityHint*/2 * 1024));
 		check(Store.IsValid());
 		check(Store.Num() == 0);
 		Store.Flush();
-		check(Store.Load(TestStorePathname));
+		verify(Store.Load(TestStorePathname));
 		check(Store.IsValid());
 		check(Store.Num() == 0);
 	}
@@ -75,21 +75,21 @@ bool FAnalyticsPropertyStoreAutomationTest::RunTest(const FString& Parameters)
 	// Test the basic operations.
 	{
 		FAnalyticsPropertyStore Store;
-		check(Store.Create(TestStorePathname, /*CapacityHint*/2 * 1024));
+		verify(Store.Create(TestStorePathname, /*CapacityHint*/2 * 1024));
 		check(Store.IsValid());
 		check(Store.Num() == 0);
 
 		// Write values.
-		check(Store.Set(I32Key,  Expected_I32) == IAnalyticsPropertyStore::EStatusCode::Success);
-		check(Store.Set(U32Key,  Expected_U32) == IAnalyticsPropertyStore::EStatusCode::Success);
-		check(Store.Set(I64Key,  Expected_I64) == IAnalyticsPropertyStore::EStatusCode::Success);
-		check(Store.Set(U64Key,  Expected_U64) == IAnalyticsPropertyStore::EStatusCode::Success);
-		check(Store.Set(FltKey,  Expected_Flt) == IAnalyticsPropertyStore::EStatusCode::Success);
-		check(Store.Set(DblKey,  Expected_Dbl) == IAnalyticsPropertyStore::EStatusCode::Success);
-		check(Store.Set(BoolKey, Expected_Bool) == IAnalyticsPropertyStore::EStatusCode::Success);
-		check(Store.Set(Str1Key, Expected_Str1, /*CapacityInChars*/25) == IAnalyticsPropertyStore::EStatusCode::Success);
-		check(Store.Set(DateKey, Expected_Date) == IAnalyticsPropertyStore::EStatusCode::Success);
-		check(Store.Set(Str2Key, Expected_Str2) == IAnalyticsPropertyStore::EStatusCode::Success);
+		verify(Store.Set(I32Key,  Expected_I32) == IAnalyticsPropertyStore::EStatusCode::Success);
+		verify(Store.Set(U32Key,  Expected_U32) == IAnalyticsPropertyStore::EStatusCode::Success);
+		verify(Store.Set(I64Key,  Expected_I64) == IAnalyticsPropertyStore::EStatusCode::Success);
+		verify(Store.Set(U64Key,  Expected_U64) == IAnalyticsPropertyStore::EStatusCode::Success);
+		verify(Store.Set(FltKey,  Expected_Flt) == IAnalyticsPropertyStore::EStatusCode::Success);
+		verify(Store.Set(DblKey,  Expected_Dbl) == IAnalyticsPropertyStore::EStatusCode::Success);
+		verify(Store.Set(BoolKey, Expected_Bool) == IAnalyticsPropertyStore::EStatusCode::Success);
+		verify(Store.Set(Str1Key, Expected_Str1, /*CapacityInChars*/25) == IAnalyticsPropertyStore::EStatusCode::Success);
+		verify(Store.Set(DateKey, Expected_Date) == IAnalyticsPropertyStore::EStatusCode::Success);
+		verify(Store.Set(Str2Key, Expected_Str2) == IAnalyticsPropertyStore::EStatusCode::Success);
 
 		// Read values
 		check(Store.Get(I32Key,  Actual_I32)  == IAnalyticsPropertyStore::EStatusCode::Success && Expected_I32 == Actual_I32);
@@ -125,53 +125,53 @@ bool FAnalyticsPropertyStoreAutomationTest::RunTest(const FString& Parameters)
 
 		// In-place update.
 		Expected_I32 = 2000;
-		check(Store.Set(I32Key, Expected_I32) == IAnalyticsPropertyStore::EStatusCode::Success);
+		verify(Store.Set(I32Key, Expected_I32) == IAnalyticsPropertyStore::EStatusCode::Success);
 		check(Store.Get(I32Key, Actual_I32) == IAnalyticsPropertyStore::EStatusCode::Success);
 		check(Expected_I32 == Actual_I32);
 
 		// In-place string update. (Enough capacity was reserved)
 		Expected_Str1 = TEXT("Hello World Update");
-		check(Store.Set(Str1Key, Expected_Str1) == IAnalyticsPropertyStore::EStatusCode::Success);
+		verify(Store.Set(Str1Key, Expected_Str1) == IAnalyticsPropertyStore::EStatusCode::Success);
 		check(Store.Get(Str1Key, Actual_Str1) == IAnalyticsPropertyStore::EStatusCode::Success && Actual_Str1 == Expected_Str1);
 		check(Store.Get(Str2Key, Actual_Str2) == IAnalyticsPropertyStore::EStatusCode::Success && Expected_Str2 == Actual_Str2); // Ensure Str1 did not overwrite Str2.
 
 		// Out-of-place string update. (Not enough capacity to perform in-place update).
 		Expected_Str2 = TEXT("Hello World Extend");
-		check(Store.Set(Str2Key, Expected_Str2) == IAnalyticsPropertyStore::EStatusCode::Success);
+		verify(Store.Set(Str2Key, Expected_Str2) == IAnalyticsPropertyStore::EStatusCode::Success);
 		check(Store.Get(Str2Key, Actual_Str2) == IAnalyticsPropertyStore::EStatusCode::Success && Expected_Str2 == Actual_Str2);
 
 		// Set conditionnaly - accepted
 		++Expected_I32;
-		check(Store.Set(I32Key, Expected_I32, [](const int32* Actual, const int32& Proposed) { return *Actual < Proposed; }) == IAnalyticsPropertyStore::EStatusCode::Success);
+		verify(Store.Set(I32Key, Expected_I32, [](const int32* Actual, const int32& Proposed) { return *Actual < Proposed; }) == IAnalyticsPropertyStore::EStatusCode::Success);
 		check(Store.Get(I32Key, Actual_I32) == IAnalyticsPropertyStore::EStatusCode::Success && Expected_I32 == Actual_I32);
 
 		// Set conditionnaly - declined
-		check(Store.Set(FltKey, Expected_Flt - 50.0f, [](const float* Actual, const float& Proposed) { return *Actual < Proposed; }) == IAnalyticsPropertyStore::EStatusCode::Declined);
+		verify(Store.Set(FltKey, Expected_Flt - 50.0f, [](const float* Actual, const float& Proposed) { return *Actual < Proposed; }) == IAnalyticsPropertyStore::EStatusCode::Declined);
 		check(Store.Get(FltKey, Actual_Flt) == IAnalyticsPropertyStore::EStatusCode::Success && Actual_Flt == Expected_Flt);
 
 		// Set conditionnaly - wrong type
-		check(Store.Set(BoolKey, Expected_Dbl, [](const double* Actual, const double& Proposed) { check(false); return *Actual < Proposed; }) == IAnalyticsPropertyStore::EStatusCode::BadType);
+		verify(Store.Set(BoolKey, Expected_Dbl, [](const double* Actual, const double& Proposed) { check(false); return *Actual < Proposed; }) == IAnalyticsPropertyStore::EStatusCode::BadType);
 
 		// Set conditionnaly - check params
 		Expected_Str1 = TEXT("Hi, this is Joe, how are you?");
-		check(Store.Set(Str1Key, Expected_Str1, [&](const FString* Actual, const FString& Proposed) { check(*Actual == Actual_Str1); check(Proposed == Expected_Str1); return true; }) == IAnalyticsPropertyStore::EStatusCode::Success);
+		verify(Store.Set(Str1Key, Expected_Str1, [&](const FString* Actual, const FString& Proposed) { check(*Actual == Actual_Str1); check(Proposed == Expected_Str1); return true; }) == IAnalyticsPropertyStore::EStatusCode::Success);
 		check(Store.Get(Str1Key, Actual_Str1) == IAnalyticsPropertyStore::EStatusCode::Success && Actual_Str1 == Expected_Str1);
 
 		// Update - accepted.
 		++Expected_I64;
-		check(Store.Update(I64Key, [](int64& InOutStored) { ++InOutStored; return true; }) == IAnalyticsPropertyStore::EStatusCode::Success);
+		verify(Store.Update(I64Key, [](int64& InOutStored) { ++InOutStored; return true; }) == IAnalyticsPropertyStore::EStatusCode::Success);
 		check(Store.Get(I64Key, Actual_I64) == IAnalyticsPropertyStore::EStatusCode::Success && Actual_I64 == Expected_I64);
 
 		// Update - declined.
-		check(Store.Update(DateKey, [](FDateTime& InOutStored) { InOutStored = FDateTime::UtcNow(); return false; }) == IAnalyticsPropertyStore::EStatusCode::Declined);
+		verify(Store.Update(DateKey, [](FDateTime& InOutStored) { InOutStored = FDateTime::UtcNow(); return false; }) == IAnalyticsPropertyStore::EStatusCode::Declined);
 		check(Store.Get(DateKey, Actual_Date) == IAnalyticsPropertyStore::EStatusCode::Success && Actual_Date == Expected_Date);
 
 		// Update - wrong type
-		check(Store.Update(Str1Key, [](uint32& InOutStored) { check(false); return true; }) == IAnalyticsPropertyStore::EStatusCode::BadType);
+		verify(Store.Update(Str1Key, [](uint32& InOutStored) { check(false); return true; }) == IAnalyticsPropertyStore::EStatusCode::BadType);
 
 		// Update - check params
 		Expected_Bool = !Expected_Bool;
-		check(Store.Update(BoolKey, [&](bool& InOutStored) { check(InOutStored == Actual_Bool); InOutStored = Expected_Bool; return true; }) == IAnalyticsPropertyStore::EStatusCode::Success);
+		verify(Store.Update(BoolKey, [&](bool& InOutStored) { check(InOutStored == Actual_Bool); InOutStored = Expected_Bool; return true; }) == IAnalyticsPropertyStore::EStatusCode::Success);
 		check(Store.Get(BoolKey, Actual_Bool) == IAnalyticsPropertyStore::EStatusCode::Success && Actual_Bool == Expected_Bool);
 
 		// Test string conversion.
@@ -230,7 +230,7 @@ bool FAnalyticsPropertyStoreAutomationTest::RunTest(const FString& Parameters)
 	// Reload the store written above and verify the values.
 	{
 		FAnalyticsPropertyStore Store;
-		check(Store.Load(TestStorePathname));
+		verify(Store.Load(TestStorePathname));
 		check(Store.IsValid());
 		check(Store.Num() == 10);
 
@@ -246,19 +246,19 @@ bool FAnalyticsPropertyStoreAutomationTest::RunTest(const FString& Parameters)
 		check(Store.Get(Str2Key, Actual_Str2) == IAnalyticsPropertyStore::EStatusCode::Success && Expected_Str2 == Actual_Str2);
 
 		// Test the remove API
-		check(Store.Remove(U32Key));
-		check(Store.Remove(I32Key));
-		check(Store.Remove(Str1Key));
-		check(Store.Remove(I64Key));
-		check(!Store.Remove(I64Key)); // Remove same key twice.
-		check(Store.Remove(U64Key));
-		check(Store.Remove(DblKey));
-		check(Store.Remove(FltKey));
-		check(Store.Remove(BoolKey));
-		check(Store.Remove(Str2Key));
-		check(!Store.Remove(Str2Key)); // Remove same key twice.
-		check(Store.Remove(DateKey));
-		check(!Store.Remove(TEXT("IDoNotExist")));
+		verify(Store.Remove(U32Key));
+		verify(Store.Remove(I32Key));
+		verify(Store.Remove(Str1Key));
+		verify(Store.Remove(I64Key));
+		verify(!Store.Remove(I64Key)); // Remove same key twice.
+		verify(Store.Remove(U64Key));
+		verify(Store.Remove(DblKey));
+		verify(Store.Remove(FltKey));
+		verify(Store.Remove(BoolKey));
+		verify(Store.Remove(Str2Key));
+		verify(!Store.Remove(Str2Key)); // Remove same key twice.
+		verify(Store.Remove(DateKey));
+		verify(!Store.Remove(TEXT("IDoNotExist")));
 
 		check(Store.Num() == 0);
 
@@ -281,7 +281,7 @@ bool FAnalyticsPropertyStoreAutomationTest::RunTest(const FString& Parameters)
 	// Reload an empty store and check it doesn't have any key.
 	{
 		FAnalyticsPropertyStore Store;
-		check(Store.Load(TestStorePathname));
+		verify(Store.Load(TestStorePathname));
 		check(Store.Num() == 0);
 
 		check(!Store.Contains(I32Key));
@@ -296,18 +296,18 @@ bool FAnalyticsPropertyStoreAutomationTest::RunTest(const FString& Parameters)
 		check(!Store.Contains(Str2Key));
 
 		// Add if null.
-		check(Store.Set(I64Key, Expected_I64, [](const int64* Actual, const int64& Proposed) { return Actual == nullptr; }) == IAnalyticsPropertyStore::EStatusCode::Success);
+		verify(Store.Set(I64Key, Expected_I64, [](const int64* Actual, const int64& Proposed) { return Actual == nullptr; }) == IAnalyticsPropertyStore::EStatusCode::Success);
 		check(Store.Get(I64Key, Actual_I64) == IAnalyticsPropertyStore::EStatusCode::Success && Actual_I64 == Expected_I64);
 
 		// Decline if null
-		check(Store.Set(I32Key, Expected_I32, [](const int32* Actual, const int32& Proposed) { return Actual != nullptr; }) == IAnalyticsPropertyStore::EStatusCode::Declined);
+		verify(Store.Set(I32Key, Expected_I32, [](const int32* Actual, const int32& Proposed) { return Actual != nullptr; }) == IAnalyticsPropertyStore::EStatusCode::Declined);
 		check(Store.Get(I32Key, Actual_I32) == IAnalyticsPropertyStore::EStatusCode::NotFound && Actual_I32 == Expected_I32);
 
 		check(Store.Num() == 1);
 
-		check(Store.Set(Str2Key, Expected_Str2) == IAnalyticsPropertyStore::EStatusCode::Success);
-		check(Store.Set(DblKey,  Expected_Dbl) == IAnalyticsPropertyStore::EStatusCode::Success);
-		check(Store.Set(BoolKey, Expected_Bool) == IAnalyticsPropertyStore::EStatusCode::Success);
+		verify(Store.Set(Str2Key, Expected_Str2) == IAnalyticsPropertyStore::EStatusCode::Success);
+		verify(Store.Set(DblKey,  Expected_Dbl) == IAnalyticsPropertyStore::EStatusCode::Success);
+		verify(Store.Set(BoolKey, Expected_Bool) == IAnalyticsPropertyStore::EStatusCode::Success);
 
 		// Remove all.
 		Store.RemoveAll();
@@ -321,7 +321,7 @@ bool FAnalyticsPropertyStoreAutomationTest::RunTest(const FString& Parameters)
 	// Corrupt a store.
 	{
 		TUniquePtr<FAnalyticsPropertyStore> Store = MakeUnique<FAnalyticsPropertyStore>();
-		check(Store->Create(TestStorePathname, /*CapacityHint*/2 * 1024));
+		verify(Store->Create(TestStorePathname, /*CapacityHint*/2 * 1024));
 
 		Store->Set(I32Key,  Expected_I32);
 		Store->Set(U32Key,  Expected_U32);
@@ -343,23 +343,23 @@ bool FAnalyticsPropertyStoreAutomationTest::RunTest(const FString& Parameters)
 		FileHandle.Reset();
 
 		Store = MakeUnique<FAnalyticsPropertyStore>();
-		check(!Store->Load(TestStorePathname)); // Checksum should fails.
+		verify(!Store->Load(TestStorePathname)); // Checksum should fails.
 	}
 
 	// Test store internal reset.
 	{
 		// Create the store.
 		FAnalyticsPropertyStore Store;
-		check(Store.Create(TestStorePathname, /*CapacityHint*/2 * 1024));
+		verify(Store.Create(TestStorePathname, /*CapacityHint*/2 * 1024));
 		Store.Set(I32Key, Expected_I32);
 		Store.Flush();
 
 		// Reload the store -> This will reset the object.
-		check(Store.Load(TestStorePathname));
+		verify(Store.Load(TestStorePathname));
 		check(Store.Get(I32Key, Actual_I32) == IAnalyticsPropertyStore::EStatusCode::Success && Expected_I32 == Actual_I32);
 
 		// Recreate the store -> This will reset the object and create an empty store.
-		check(Store.Create(TestStorePathname, /*CapacityHint*/2 * 1024));
+		verify(Store.Create(TestStorePathname, /*CapacityHint*/2 * 1024));
 		check(Store.Get(I32Key, Actual_I32) == IAnalyticsPropertyStore::EStatusCode::NotFound);
 		check(!Store.Contains(I32Key));
 	}
@@ -368,12 +368,12 @@ bool FAnalyticsPropertyStoreAutomationTest::RunTest(const FString& Parameters)
 	{
 		// Create the store.
 		FAnalyticsPropertyStore Store;
-		check(Store.Create(TestStorePathname, /*CapacityHint*/2 * 1024));
+		verify(Store.Create(TestStorePathname, /*CapacityHint*/2 * 1024));
 		Store.Set(I32Key, Expected_I32);
 		Store.Flush(/*bAsync*/true);
 
 		// Reload the store -> This will wait until the async task complates and reset the object.
-		check(Store.Load(TestStorePathname));
+		verify(Store.Load(TestStorePathname));
 		check(Store.Get(I32Key, Actual_I32) == IAnalyticsPropertyStore::EStatusCode::Success && Expected_I32 == Actual_I32);
 	}
 
@@ -381,12 +381,12 @@ bool FAnalyticsPropertyStoreAutomationTest::RunTest(const FString& Parameters)
 	{
 		// Create the store.
 		FAnalyticsPropertyStore Store;
-		check(Store.Create(TestStorePathname, /*CapacityHint*/2 * 1024));
+		verify(Store.Create(TestStorePathname, /*CapacityHint*/2 * 1024));
 		Store.Set(I32Key, Expected_I32);
 		Store.Flush(/*bAsync*/true, FTimespan::FromMilliseconds(100));
 
 		// Reload the store -> This will wait until the async task complates and reset the object.
-		check(Store.Load(TestStorePathname));
+		verify(Store.Load(TestStorePathname));
 		check(Store.Get(I32Key, Actual_I32) == IAnalyticsPropertyStore::EStatusCode::Success && Expected_I32 == Actual_I32);
 	}
 
diff --git a/Engine/Source/Runtime/Json/Private/Tests/JsonTests.cpp b/Engine/Source/Runtime/Json/Private/Tests/JsonTests.cpp
--- a/Engine/Source/Runtime/Json/Private/Tests/JsonTests.cpp
+++ b/Engine/Source/Runtime/Json/Private/Tests/JsonTests.cpp
@@ -35,7 +35,7 @@ bool FJsonAutomationTest::RunTest(const FString& Parameters)
 		TSharedRef< TJsonReader<> > Reader = TJsonReaderFactory<>::Create( InputString );
 
 		TSharedPtr<FJsonObject> Object;
-		check( FJsonSerializer::Deserialize( Reader, Object ) == false );
+		verify( FJsonSerializer::Deserialize( Reader, Object ) == false );
 		check( !Object.IsValid() );
 	}
 
@@ -45,12 +45,12 @@ bool FJsonAutomationTest::RunTest(const FString& Parameters)
 		TSharedRef< TJsonReader<> > Reader = TJsonReaderFactory<>::Create( InputString );
 
 		TSharedPtr<FJsonObject> Object;
-		check( FJsonSerializer::Deserialize( Reader, Object ) );
+		verify( FJsonSerializer::Deserialize( Reader, Object ) );
 		check( Object.IsValid() );
 
 		FString OutputString;
 		TSharedRef< FCondensedJsonStringWriter > Writer = FCondensedJsonStringWriterFactory::Create( &OutputString );
-		check( FJsonSerializer::Serialize( Object.ToSharedRef(), Writer ) );
+		verify( FJsonSerializer::Serialize( Object.ToSharedRef(), Writer ) );
 		check( InputString == OutputString );
 	}
 
@@ -60,12 +60,12 @@ bool FJsonAutomationTest::RunTest(const FString& Parameters)
 		TSharedRef< TJsonReader<> > Reader = TJsonReaderFactory<>::Create( InputString );
 
 		TArray< TSharedPtr<FJsonValue> > Array;
-		check( FJsonSerializer::Deserialize( Reader, Array ) );
+		verify( FJsonSerializer::Deserialize( Reader, Array ) );
 		check( Array.Num() == 0 );
 
 		FString OutputString;
 		TSharedRef< FCondensedJsonStringWriter > Writer = FCondensedJsonStringWriterFactory::Create( &OutputString );
-		check( FJsonSerializer::Serialize( Array, Writer ) );
+		verify( FJsonSerializer::Serialize( Array, Writer ) );
 		check( InputString == OutputString );
 	}
 
@@ -94,7 +94,7 @@ bool FJsonAutomationTest::RunTest(const FString& Parameters)
 
 		FString OutputString;
 		TSharedRef< FCondensedJsonStringWriter > Writer = FCondensedJsonStringWriterFactory::Create( &OutputString );
-		check( FJsonSerializer::Serialize( Array, Writer ) );
+		verify( FJsonSerializer::Serialize( Array, Writer ) );
 		check( InputString == OutputString );
 	}
 
@@ -312,7 +312,7 @@ bool FJsonAutomationTest::RunTest(const FString& Parameters)
 
 		FString OutputString;
 		TSharedRef< FCondensedJsonStringWriter > Writer = FCondensedJsonStringWriterFactory::Create( &OutputString );
-		check( FJsonSerializer::Serialize( Object.ToSharedRef(), Writer ) );
+		verify( FJsonSerializer::Serialize( Object.ToSharedRef(), Writer ) );
 
 		const FString TestOutput =
 			TEXT(
@@ -353,7 +353,7 @@ bool FJsonAutomationTest::RunTest(const FString& Parameters)
 
 		FString OutputString;
 		TSharedRef< FCondensedJsonStringWriter > Writer = FCondensedJsonStringWriterFactory::Create( &OutputString );
-		check( FJsonSerializer::Serialize( Object.ToSharedRef(), Writer ) );
+		verify( FJsonSerializer::Serialize( Object.ToSharedRef(), Writer ) );
 
 		// %g isn't standardized, so we use the same %g format that is used inside PrintJson instead of hardcoding the values here
 		const FString TestOutput = FString::Printf(
@@ -409,7 +409,7 @@ bool FJsonAutomationTest::RunTest(const FString& Parameters)
 
 		FString OutputString;
 		TSharedRef< FCondensedJsonStringWriter > Writer = FCondensedJsonStringWriterFactory::Create( &OutputString );
-		check( FJsonSerializer::Serialize( Object.ToSharedRef(), Writer ) );
+		verify( FJsonSerializer::Serialize( Object.ToSharedRef(), Writer ) );
 
 		const FString TestOutput =
 			TEXT(
@@ -488,7 +488,7 @@ bool FJsonAutomationTest::RunTest(const FString& Parameters)
 
 		FString OutputString;
 		TSharedRef< FCondensedJsonStringWriter > Writer = FCondensedJsonStringWriterFactory::Create( &OutputString );
-		check( FJsonSerializer::Serialize( Object.ToSharedRef(), Writer ) );
+		verify( FJsonSerializer::Serialize( Object.ToSharedRef(), Writer ) );
 		check(OutputString == InputString);
 	}
 
@@ -547,7 +547,7 @@ bool FJsonAutomationTest::RunTest(const FString& Parameters)
 
 		FString OutputString;
 		TSharedRef< FCondensedJsonStringWriter > Writer = FCondensedJsonStringWriterFactory::Create( &OutputString );
-		check( FJsonSerializer::Serialize( Object.ToSharedRef(), Writer ) );
+		verify( FJsonSerializer::Serialize( Object.ToSharedRef(), Writer ) );
 		check(OutputString == InputString);
 	}
 
@@ -577,12 +577,12 @@ bool FJsonAutomationTest::RunTest(const FString& Parameters)
 		TSharedRef< TJsonReader<> > Reader = TJsonReaderFactory<>::Create( InputString );
 
 		TSharedPtr<FJsonObject> Object;
-		check( FJsonSerializer::Deserialize( Reader, Object ) );
+		verify( FJsonSerializer::Deserialize( Reader, Object ) );
 		check( Object.IsValid() );
 
 		FString OutputString;
 		TSharedRef< FPrettyJsonStringWriter > Writer = FPrettyJsonStringWriterFactory::Create( &OutputString );
-		check( FJsonSerializer::Serialize( Object.ToSharedRef(), Writer ) );
+		verify( FJsonSerializer::Serialize( Object.ToSharedRef(), Writer ) );
 		check(OutputString == InputString);
 	}
 	  
@@ -605,19 +605,19 @@ bool FJsonAutomationTest::RunTest(const FString& Parameters)
 		TSharedRef< TJsonReader<> > Reader = TJsonReaderFactory<>::Create( InputString );
 
 		EJsonNotation Notation = EJsonNotation::Null;
-		check( Reader->ReadNext( Notation ) && Notation == EJsonNotation::ObjectStart );
+		verify( Reader->ReadNext( Notation ) && Notation == EJsonNotation::ObjectStart );
 		check( Reader->GetLineNumber() == 1 && Reader->GetCharacterNumber() == 1 );
 
-		check( Reader->ReadNext( Notation ) && Notation == EJsonNotation::String );
+		verify( Reader->ReadNext( Notation ) && Notation == EJsonNotation::String );
 		check( Reader->GetLineNumber() == 2 && Reader->GetCharacterNumber() == 17 );
 
-		check( Reader->ReadNext( Notation ) && Notation == EJsonNotation::ArrayStart );
+		verify( Reader->ReadNext( Notation ) && Notation == EJsonNotation::ArrayStart );
 		check( Reader->GetLineNumber() == 4 && Reader->GetCharacterNumber() == 2 );
 
-		check( Reader->ReadNext( Notation ) && Notation == EJsonNotation::Number );
+		verify( Reader->ReadNext( Notation ) && Notation == EJsonNotation::Number );
 		check( Reader->GetLineNumber() == 5 && Reader->GetCharacterNumber() == 7 );
 
-		check( Reader->ReadNext( Notation ) && Notation == EJsonNotation::Boolean );
+		verify( Reader->ReadNext( Notation ) && Notation == EJsonNotation::Boolean );
 		check( Reader->GetLineNumber() == 6 && Reader->GetCharacterNumber() == 6 );
 	}
 
@@ -742,7 +742,7 @@ bool FJsonAutomationTest::RunTest(const FString& Parameters)
 		TSharedRef< TJsonReader<> > Reader = TJsonReaderFactory<>::Create( FailureInputs[i] );
 
 		TSharedPtr<FJsonObject> Object;
-		check( FJsonSerializer::Deserialize( Reader, Object ) == false );
+		verify( FJsonSerializer::Deserialize( Reader, Object ) == false );
 		check( !Object.IsValid() );
 	}
 
diff --git a/Engine/Source/Runtime/Online/XMPP/Private/XmppJingle/XmppMessagesJingle.cpp b/Engine/Source/Runtime/Online/XMPP/Private/XmppJingle/XmppMessagesJingle.cpp
--- a/Engine/Source/Runtime/Online/XMPP/Private/XmppJingle/XmppMessagesJingle.cpp
+++ b/Engine/Source/Runtime/Online/XMPP/Private/XmppJingle/XmppMessagesJingle.cpp
@@ -314,7 +314,7 @@ bool FXmppMessagesJingle::SendMessage(const FXmppUserJid& RecipientId, const FSt
 {
 	FString SerializedPayload;
 	auto JsonWriter = TJsonWriterFactory<TCHAR, TCondensedJsonPrintPolicy<TCHAR>>::Create(&SerializedPayload);
-	check(FJsonSerializer::Serialize(Payload, JsonWriter));
+	verify(FJsonSerializer::Serialize(Payload, JsonWriter));
 
 	return SendMessage(RecipientId, Type, SerializedPayload, true);
 }
diff --git a/Engine/Source/Runtime/Online/XMPP/Private/XmppStrophe/XmppMessagesStrophe.cpp b/Engine/Source/Runtime/Online/XMPP/Private/XmppStrophe/XmppMessagesStrophe.cpp
--- a/Engine/Source/Runtime/Online/XMPP/Private/XmppStrophe/XmppMessagesStrophe.cpp
+++ b/Engine/Source/Runtime/Online/XMPP/Private/XmppStrophe/XmppMessagesStrophe.cpp
@@ -188,7 +188,7 @@ bool FXmppMessagesStrophe::SendMessage(const FXmppUserJid& RecipientId, const FS
 {
 	FString SerializedPayload;
 	auto JsonWriter = TJsonWriterFactory<TCHAR, TCondensedJsonPrintPolicy<TCHAR>>::Create(&SerializedPayload);
-	check(FJsonSerializer::Serialize(Payload, JsonWriter));
+	verify(FJsonSerializer::Serialize(Payload, JsonWriter));
 
 	return SendMessage(RecipientId, Type, SerializedPayload, true);
 }
