Fix Test build failing with bForceCompileDevelopmentAutomationTests (Contributed by [REDACTED_PERSON])

PR #9678

diff --git a/Engine/Source/Runtime/Analytics/Analytics/Private/Tests/AnalyticsPropertyStoreTests.cpp b/Engine/Source/Runtime/Analytics/Analytics/Private/Tests/AnalyticsPropertyStoreTests.cpp
--- a/Engine/Source/Runtime/Analytics/Analytics/Private/Tests/AnalyticsPropertyStoreTests.cpp
+++ b/Engine/Source/Runtime/Analytics/Analytics/Private/Tests/AnalyticsPropertyStoreTests.cpp
@@ -37,6 +37,10 @@ bool FAnalyticsPropertyStoreAutomationTest::RunTest(const FString& Parameters)
 	FString Expected_Str2 = "Just Do It";
 	FDateTime Expected_Date = FDateTime::UtcNow();
 
+// When DO_CHECK is disabled (by default in Test and Shipping builds), the check() macro is a no-op.
+// The variables below are only used in check().
+// To avoid "error C4101: unreferenced local variable", only declare them when check() actually uses them.
+#if DO_CHECK
 	int32 Actual_I32;
 	uint32 Actual_U32;
 	int64 Actual_I64;
@@ -48,6 +52,7 @@ bool FAnalyticsPropertyStoreAutomationTest::RunTest(const FString& Parameters)
 	FString Actual_Str2;
 	FDateTime Actual_Date;
 	int32 NotFoundDummy;
+#endif
 
 	// Ensure the automation directory exists.
 	IFileManager::Get().MakeDirectory(*FPaths::AutomationTransientDir(), /*Tree*/true);
