Remove ProjectName from telemetry User-Agent field
#tests compile

diff --git a/Engine/Source/Runtime/Analytics/AnalyticsET/Private/IAnalyticsProviderET.cpp b/Engine/Source/Runtime/Analytics/AnalyticsET/Private/IAnalyticsProviderET.cpp
--- a/Engine/Source/Runtime/Analytics/AnalyticsET/Private/IAnalyticsProviderET.cpp
+++ b/Engine/Source/Runtime/Analytics/AnalyticsET/Private/IAnalyticsProviderET.cpp
@@ -387,6 +387,13 @@ void FAnalyticsProviderET::FlushEventsOnce()
 			// Create/send Http request for an event
 			TSharedRef<IHttpRequest, ESPMode::ThreadSafe> HttpRequest = CreateRequest();
 			HttpRequest->SetHeader(TEXT("Content-Type"), TEXT("application/json; charset=utf-8"));
+			// Want to avoid putting the project name into the User-Agent, because for some apps (like the editor), the project name is private info.
+			// Code Pulled from FGenericPlatformHttp::GetDefaultUserAgent, but with project name pulled out
+			HttpRequest->SetHeader(TEXT("User-Agent"), FString::Printf(TEXT("%s/%s %s/%s"),
+				TEXT("PROJECTNAME"),
+				*FPlatformHttp::EscapeUserAgentString(FApp::GetBuildVersion()),
+				*FPlatformHttp::EscapeUserAgentString(FString(FPlatformProperties::IniPlatformName())),
+				*FPlatformHttp::EscapeUserAgentString(FPlatformMisc::GetOSVersion())));
 			HttpRequest->SetURL(Config.APIServerET / URLPath);
 			HttpRequest->SetVerb(TEXT("POST"));
 			HttpRequest->SetContent(MoveTemp(Payload));
