Truncation warning fixes for multiple core modules:
Analytics, AnalyticsET, Cbor, Serialization, XMLParser, DerivedDataEditor, DirectoryWatcher, LogVisualizer, OutputLog, InputCore

diff --git a/Engine/Source/Developer/DirectoryWatcher/Private/FileCacheUtilities.cpp b/Engine/Source/Developer/DirectoryWatcher/Private/FileCacheUtilities.cpp
--- a/Engine/Source/Developer/DirectoryWatcher/Private/FileCacheUtilities.cpp
+++ b/Engine/Source/Developer/DirectoryWatcher/Private/FileCacheUtilities.cpp
@@ -16,8 +16,8 @@ bool MatchExtensionString(const TCHAR* Filename, const TCHAR* Extensions)
 	{
 		return false;
 	}
-
-	const int32 ExtLength = StrLength - (Ext - Filename);
+	
+	const int32 ExtLength = StrLength - UE_PTRDIFF_TO_INT32(Ext - Filename);
 	const TCHAR* Search = FCString::Strchr(Extensions, QueryChar);
 
 	while (Search && *(++Search) != '\0')
diff --git a/Engine/Source/Developer/OutputLog/Private/SDeviceOutputLog.cpp b/Engine/Source/Developer/OutputLog/Private/SDeviceOutputLog.cpp
--- a/Engine/Source/Developer/OutputLog/Private/SDeviceOutputLog.cpp
+++ b/Engine/Source/Developer/OutputLog/Private/SDeviceOutputLog.cpp
@@ -68,8 +68,8 @@ void SDeviceOutputLog::Construct( const FArguments& InArgs )
 						.AutoWidth()
 						[
 							SNew(SBox)
-							.WidthOverride(16)
-							.HeightOverride(16)
+							.WidthOverride(16.f)
+							.HeightOverride(16.f)
 							[
 								SNew(SImage).Image(this, &SDeviceOutputLog::GetSelectedTargetDeviceBrush)
 							]
@@ -342,8 +342,8 @@ TSharedRef<SWidget> SDeviceOutputLog::GenerateWidgetForDeviceComboBox(const FTar
 			.AutoWidth()
 			[
 				SNew(SBox)
-				.WidthOverride(24)
-				.HeightOverride(24)
+				.WidthOverride(24.f)
+				.HeightOverride(24.f)
 				[
 					SNew(SImage).Image(GetTargetDeviceBrush(DeviceEntry))
 				]
diff --git a/Engine/Source/Developer/OutputLog/Private/SOutputLog.cpp b/Engine/Source/Developer/OutputLog/Private/SOutputLog.cpp
--- a/Engine/Source/Developer/OutputLog/Private/SOutputLog.cpp
+++ b/Engine/Source/Developer/OutputLog/Private/SOutputLog.cpp
@@ -181,7 +181,7 @@ void SConsoleInputBox::Construct(const FArguments& InArgs)
 				SNew(SComboButton)
 				.IsEnabled(this, &SConsoleInputBox::IsCommandExecutorMenuEnabled)
 				.ComboButtonStyle(FOutputLogStyle::Get(), "SimpleComboButton")
-				.ContentPadding(0)
+				.ContentPadding(0.f)
 				.OnGetMenuContent(this, &SConsoleInputBox::GetCommandExecutorMenuContent)
 				.ButtonContent()
 				[
@@ -228,8 +228,8 @@ void SConsoleInputBox::Construct(const FArguments& InArgs)
 			.Padding( FMargin(2) )
 			[
 				SNew(SBox)
-				.HeightOverride(250) // avoids flickering, ideally this would be adaptive to the content without flickering
-				.MinDesiredWidth(300)
+				.HeightOverride(250.f) // avoids flickering, ideally this would be adaptive to the content without flickering
+				.MinDesiredWidth(300.f)
 				.MaxDesiredWidth(this, &SConsoleInputBox::GetSelectionListMaxWidth)
 				[
 					SAssignNew(SuggestionListView, SListView< TSharedPtr<FString> >)
@@ -856,7 +856,7 @@ float FOutputLogTextLayoutMarshaller::GetCategoryHue(FName CategoryName)
 	else
 	{
 		FRandomStream RNG(GetTypeHash(CategoryName));
-		const float Hue = RNG.FRandRange(0.0f, 360.0f);
+		const float Hue = (float)RNG.FRandRange(0.0, 360.0);
 		CategoryHueMap.Add(CategoryName, Hue);
 		return Hue;
 	}
diff --git a/Engine/Source/Editor/DerivedDataEditor/Private/DerivedDataInformation.cpp b/Engine/Source/Editor/DerivedDataEditor/Private/DerivedDataInformation.cpp
--- a/Engine/Source/Editor/DerivedDataEditor/Private/DerivedDataInformation.cpp
+++ b/Engine/Source/Editor/DerivedDataEditor/Private/DerivedDataInformation.cpp
@@ -61,7 +61,7 @@ double FDerivedDataInformation::GetCacheActivitySizeBytes(bool bGet, bool bLocal
 	}
 #endif // ENABLE_COOK_STATS
 
-	return TotalBytes;
+	return (double)TotalBytes;
 }
 
 
diff --git a/Engine/Source/Editor/DerivedDataEditor/Private/SDerivedDataDialogs.cpp b/Engine/Source/Editor/DerivedDataEditor/Private/SDerivedDataDialogs.cpp
--- a/Engine/Source/Editor/DerivedDataEditor/Private/SDerivedDataDialogs.cpp
+++ b/Engine/Source/Editor/DerivedDataEditor/Private/SDerivedDataDialogs.cpp
@@ -666,8 +666,8 @@ TSharedRef<SWidget> SDerivedDataCacheStatisticsDialog::GetGridPanel()
 		const int64 TotalRequests = TotalGetHits + TotalGetMisses;
 		const double HitRate = TotalRequests > 0 ? 100.0 * TotalGetHits / TotalRequests : 0.0;
 
-		const double TotalGetMB = FUnitConversion::Convert(TotalGetBytes, EUnit::Bytes, EUnit::Megabytes);
-		const double TotalPutMB = FUnitConversion::Convert(TotalPutBytes, EUnit::Bytes, EUnit::Megabytes);
+		const double TotalGetMB = FUnitConversion::Convert((double)TotalGetBytes, EUnit::Bytes, EUnit::Megabytes);
+		const double TotalPutMB = FUnitConversion::Convert((double)TotalPutBytes, EUnit::Bytes, EUnit::Megabytes);
 
 		SumTotalGetMB += TotalGetMB;
 		SumTotalPutMB += TotalPutMB;
diff --git a/Engine/Source/Runtime/Analytics/Analytics/Private/AnalyticsPropertyStore.cpp b/Engine/Source/Runtime/Analytics/Analytics/Private/AnalyticsPropertyStore.cpp
--- a/Engine/Source/Runtime/Analytics/Analytics/Private/AnalyticsPropertyStore.cpp
+++ b/Engine/Source/Runtime/Analytics/Analytics/Private/AnalyticsPropertyStore.cpp
@@ -10,6 +10,7 @@
 #include "Misc/ScopeExit.h"
 #include "Misc/ScopeLock.h"
 #include "ProfilingDebugging/CpuProfilerTrace.h"
+#include "Templates/UnrealTemplate.h"
 
 namespace AnalyticsPropertyStoreUtils
 {
@@ -631,7 +632,7 @@ IAnalyticsPropertyStore::EStatusCode FAnalyticsPropertyStore::SetFixedSizeValueI
 	else
 	{
 		// Insert the new record at the end.
-		uint32 NewRecordOffset = StorageWriter.TotalSize();
+		int64 NewRecordOffset = StorageWriter.TotalSize();
 		StorageWriter.Seek(NewRecordOffset);
 
 		// Write the type.
@@ -644,7 +645,7 @@ IAnalyticsPropertyStore::EStatusCode FAnalyticsPropertyStore::SetFixedSizeValueI
 		StorageWriter << const_cast<FString&>(Key);
 
 		// Cache the offset to the newly inserted record.
-		NameOffsetMap.Emplace(Key, NewRecordOffset);
+		NameOffsetMap.Emplace(Key, IntCastChecked<uint32>(NewRecordOffset));
 		return EStatusCode::Success;
 	}
 }
@@ -677,7 +678,7 @@ IAnalyticsPropertyStore::EStatusCode FAnalyticsPropertyStore::SetStringValueInte
 	auto InsertFn = [this](const FString& Key, ETypeCode TypeCode, const void* Bytes, uint32 ByteCount, uint32 Capacity)
 	{
 		// Insert the new record at the end.
-		uint32 InsertOffset = StorageWriter.TotalSize();
+		int64 InsertOffset = StorageWriter.TotalSize();
 		StorageWriter.Seek(InsertOffset);
 
 		// Write the value type.
@@ -702,7 +703,7 @@ IAnalyticsPropertyStore::EStatusCode FAnalyticsPropertyStore::SetStringValueInte
 		StorageWriter << const_cast<FString&>(Key);
 
 		// Cache the offset to the newly inserted record.
-		NameOffsetMap.Emplace(Key, InsertOffset);
+		NameOffsetMap.Emplace(Key, IntCastChecked<uint32>(InsertOffset));
 		return EStatusCode::Success;
 	};
 
@@ -985,7 +986,7 @@ bool FAnalyticsPropertyStore::Flush(bool bAsync, const FTimespan& Timeout)
 					return true;
 				}
 			}
-			FPlatformProcess::Sleep(FTimespan::FromMilliseconds(1).GetTotalSeconds());
+			FPlatformProcess::Sleep(FloatCastChecked<float>(FTimespan::FromMilliseconds(1).GetTotalSeconds(), 1./1000.));
 		} while (FPlatformTime::Seconds() - StartTimeSecs < Timeout.GetTotalSeconds());
 
 		return false;
diff --git a/Engine/Source/Runtime/Analytics/AnalyticsET/Private/AnalyticsProviderETEventCache.cpp b/Engine/Source/Runtime/Analytics/AnalyticsET/Private/AnalyticsProviderETEventCache.cpp
--- a/Engine/Source/Runtime/Analytics/AnalyticsET/Private/AnalyticsProviderETEventCache.cpp
+++ b/Engine/Source/Runtime/Analytics/AnalyticsET/Private/AnalyticsProviderETEventCache.cpp
@@ -109,7 +109,7 @@ namespace EventCacheStatic
 		while (!bWroteFullString)
 		{
 			// Give some padding. ensure we add at least one char.
-			const int32 StrLen = (int32)(Len + FMath::Max(1.f, Len * (SizeMultiplier)));
+			const int32 StrLen = Len + (int32)FMath::Max(1.f, (float)Len * SizeMultiplier);
 			// make space for the string
 			UTF8Stream.SetNumUninitialized(OldLen + StrLen, false);
 			// convert it to UTF8
@@ -122,9 +122,9 @@ namespace EventCacheStatic
 			else
 			{
 				// we overflowed our buffer. Must be lots of multibyte chars. double the slack and try again.
-				SizeMultiplier *= 2.0;
+				SizeMultiplier *= 2.0f;
 				// if we grow too much, give up and compute the true chars needed.
-				if (SizeMultiplier >= 2.0)
+				if (SizeMultiplier >= 2.0f)
 				{
 					const int32 ActualCharsNeeded = FPlatformString::ConvertedLength<UTF8CHAR>(Str, Len);
 					UTF8Stream.SetNumUninitialized(OldLen + ActualCharsNeeded, false);
@@ -181,7 +181,7 @@ namespace EventCacheStatic
 
 	inline void InitializePayloadBuffer(TArray<uint8>& Buffer, int32 MaximumPayloadSize)
 	{
-		Buffer.Reserve(MaximumPayloadSize * 1.2);
+		Buffer.Reserve((int32)(MaximumPayloadSize * 1.2));
 		// we are going to write UTF8 directly into our payload buffer.
 		AppendString(Buffer, PayloadTemplate, PayloadTemplateLength);
 	}
@@ -385,9 +385,9 @@ void FAnalyticsProviderETEventCache::QueueFlush()
 
 	// see if it took too long or we have a really large payload. If so, log out the events.
 	const double EndTime = FPlatformTime::Seconds();
-	if ((EndTime - StartTime) > EventCacheStatic::PayloadFlushTimeSecForWarning || CachedEventUTF8Stream.Num() > (int32)(MaximumPayloadSize * EventCacheStatic::PayloadPercentageOfMaxForWarning))
+	if ((EndTime - StartTime) > EventCacheStatic::PayloadFlushTimeSecForWarning || CachedEventUTF8Stream.Num() > (int32)((float)MaximumPayloadSize * EventCacheStatic::PayloadPercentageOfMaxForWarning))
 	{
-		UE_LOG(LogAnalytics, Warning, TEXT("EventCache either took too long to flush (%.3f ms) or had a very large payload (%.3f KB, %d events). Listing events in the payload for investigation:"), (EndTime-StartTime) * 1000, CachedEventUTF8Stream.Num() / 1024.f, CachedEventEntries.Num());
+		UE_LOG(LogAnalytics, Warning, TEXT("EventCache either took too long to flush (%.3f ms) or had a very large payload (%.3f KB, %d events). Listing events in the payload for investigation:"), (EndTime-StartTime) * 1000, (float)CachedEventUTF8Stream.Num() / 1024.f, CachedEventEntries.Num());
 		for (const FAnalyticsEventEntry& Entry : CachedEventEntries)
 		{
 			UE_LOG(LogAnalytics, Warning, TEXT("    %s,%d"), *Entry.EventName, Entry.EventSizeChars);
diff --git a/Engine/Source/Runtime/Analytics/AnalyticsET/Private/IAnalyticsProviderET.cpp b/Engine/Source/Runtime/Analytics/AnalyticsET/Private/IAnalyticsProviderET.cpp
--- a/Engine/Source/Runtime/Analytics/AnalyticsET/Private/IAnalyticsProviderET.cpp
+++ b/Engine/Source/Runtime/Analytics/AnalyticsET/Private/IAnalyticsProviderET.cpp
@@ -244,7 +244,7 @@ bool FAnalyticsProviderET::Tick(float DeltaSeconds)
 	// On servers where there may be dozens of provider instances, this will spread out the cost a bit.
 	// If caching is disabled, we still want events to be flushed immediately, so we are only guarding the flush calls from tick,
 	// any other calls to flush are allowed to happen in the same frame.
-	static uint32 LastFrameCounterFlushed = 0;
+	static uint64 LastFrameCounterFlushed = 0;
 
 	const bool bHadFlushesQueued = EventCache.HasFlushesQueued();
 	const bool bShouldFlush = bHadFlushesQueued || (EventCache.CanFlush() && Now >= NextEventFlushTime);
@@ -264,7 +264,7 @@ bool FAnalyticsProviderET::Tick(float DeltaSeconds)
 			// try to keep on the same cadence when flushing, since we could miss our window by several frames.
 			if (!bHadFlushesQueued && Now >= NextEventFlushTime)
 			{
-				const float Multiplier = (int)((Now - NextEventFlushTime) / FlushIntervalSec) + 1.f;
+				const float Multiplier = (FloatCastChecked<float>(Now - NextEventFlushTime, 1./16.) / FlushIntervalSec) + 1.f;
 				NextEventFlushTime += Multiplier * FlushIntervalSec;
 			}
 		}
diff --git a/Engine/Source/Runtime/Cbor/Private/CborReader.cpp b/Engine/Source/Runtime/Cbor/Private/CborReader.cpp
--- a/Engine/Source/Runtime/Cbor/Private/CborReader.cpp
+++ b/Engine/Source/Runtime/Cbor/Private/CborReader.cpp
@@ -148,9 +148,10 @@ bool FCborReader::ReadNext(FCborContext& OutContext)
 			else
 			{
 				OutContext.Length = ReadUIntValue(OutContext, *Stream);
-				OutContext.RawTextValue.SetNumUninitialized(OutContext.Length + 1); // Length doesn't count the null terminating character
+				int32 StringLength = IntCastChecked<int32>(OutContext.Length);
+				OutContext.RawTextValue.SetNumUninitialized(StringLength + 1); // Length doesn't count the null terminating character
 				Stream->Serialize(OutContext.RawTextValue.GetData(), OutContext.Length);
-				OutContext.RawTextValue[OutContext.Length] = '\0';
+				OutContext.RawTextValue[StringLength] = '\0';
 			}
 			break;
 		case ECborCode::Array:
diff --git a/Engine/Source/Runtime/Cbor/Private/CborWriter.cpp b/Engine/Source/Runtime/Cbor/Private/CborWriter.cpp
--- a/Engine/Source/Runtime/Cbor/Private/CborWriter.cpp
+++ b/Engine/Source/Runtime/Cbor/Private/CborWriter.cpp
@@ -160,21 +160,21 @@ FCborHeader FCborWriter::WriteUIntValue(FCborHeader Header, FArchive& Ar, uint64
 	{
 		Header.Set(Header.MajorType() | ECborCode::Value_1Byte);
 		Ar << Header;
-		uint8 Temp = Value;
+		uint8 Temp = (uint8)(Value);
 		Ar << Temp;
 	}
 	else if (Value < 65536)
 	{
 		Header.Set((uint8)(Header.MajorType() | ECborCode::Value_2Bytes));
 		Ar << Header;
-		uint16 Temp = Value;
+		uint16 Temp = (uint16)Value;
 		Ar << Temp;
 	}
 	else if (Value < 0x100000000L)
 	{
 		Header.Set((uint8)(Header.MajorType() | ECborCode::Value_4Bytes));
 		Ar << Header;
-		uint32 Temp = Value;
+		uint32 Temp = (uint32)Value;
 		Ar << Temp;
 	}
 	else
diff --git a/Engine/Source/Runtime/InputCore/Private/InputCoreTypes.cpp b/Engine/Source/Runtime/InputCore/Private/InputCoreTypes.cpp
--- a/Engine/Source/Runtime/InputCore/Private/InputCoreTypes.cpp
+++ b/Engine/Source/Runtime/InputCore/Private/InputCoreTypes.cpp
@@ -1572,7 +1572,7 @@ FKey FInputKeyManager::GetKeyFromCodes( const uint32 KeyCode, const uint32 CharC
 	if (KeyPtr == nullptr && CharCode > 32)
 	{
 		FKey NewKey(*FString::Printf(TEXT("%s%d"), FKey::SyntheticCharPrefix, CharCode));
-		EKeys::AddKey(FKeyDetails(NewKey, FText::AsCultureInvariant(FString::Chr(CharCode)), FKeyDetails::NotBlueprintBindableKey | FKeyDetails::NotActionBindableKey));
+		EKeys::AddKey(FKeyDetails(NewKey, FText::AsCultureInvariant(FString::Chr((TCHAR)CharCode)), FKeyDetails::NotBlueprintBindableKey | FKeyDetails::NotActionBindableKey));
 		const_cast<FInputKeyManager*>(this)->KeyMapCharToEnum.Add(CharCode, NewKey);
 		return NewKey;
 	}
diff --git a/Engine/Source/Runtime/InputCore/Private/Windows/WindowsPlatformInput.cpp b/Engine/Source/Runtime/InputCore/Private/Windows/WindowsPlatformInput.cpp
--- a/Engine/Source/Runtime/InputCore/Private/Windows/WindowsPlatformInput.cpp
+++ b/Engine/Source/Runtime/InputCore/Private/Windows/WindowsPlatformInput.cpp
@@ -111,7 +111,7 @@ uint32 FWindowsPlatformInput::GetKeyMap( uint32* KeyCodes, FString* KeyNames, ui
 
 		for (auto It(ScanToVKMap.CreateConstIterator()); It; ++It)
 		{
-			ADDKEYMAP(It.Value(), FString::Chr(It.Key()));
+			ADDKEYMAP(It.Value(), FString::Chr((TCHAR)It.Key()));
 		}
 	}
 
diff --git a/Engine/Source/Runtime/Serialization/Private/Backends/CborStructDeserializerBackend.cpp b/Engine/Source/Runtime/Serialization/Private/Backends/CborStructDeserializerBackend.cpp
--- a/Engine/Source/Runtime/Serialization/Private/Backends/CborStructDeserializerBackend.cpp
+++ b/Engine/Source/Runtime/Serialization/Private/Backends/CborStructDeserializerBackend.cpp
@@ -203,7 +203,7 @@ bool FCborStructDeserializerBackend::ReadProperty(FProperty* Property, FProperty
 				return false;
 			}
 
-			int32 Value = ByteProperty->Enum->GetValueByName(*StringValue);
+			int64 Value = ByteProperty->Enum->GetValueByName(*StringValue);
 			if (Value == INDEX_NONE)
 			{
 				return false;
diff --git a/Engine/Source/Runtime/Serialization/Private/Backends/JsonStructDeserializerBackend.cpp b/Engine/Source/Runtime/Serialization/Private/Backends/JsonStructDeserializerBackend.cpp
--- a/Engine/Source/Runtime/Serialization/Private/Backends/JsonStructDeserializerBackend.cpp
+++ b/Engine/Source/Runtime/Serialization/Private/Backends/JsonStructDeserializerBackend.cpp
@@ -193,7 +193,7 @@ bool FJsonStructDeserializerBackend::ReadProperty( FProperty* Property, FPropert
 					return false;
 				}
 
-				int32 Value = ByteProperty->Enum->GetValueByName(*StringValue);
+				int64 Value = ByteProperty->Enum->GetValueByName(*StringValue);
 				if (Value == INDEX_NONE)
 				{
 					return false;
diff --git a/Engine/Source/Runtime/Serialization/Private/Tests/StructSerializerTest.cpp b/Engine/Source/Runtime/Serialization/Private/Tests/StructSerializerTest.cpp
--- a/Engine/Source/Runtime/Serialization/Private/Tests/StructSerializerTest.cpp
+++ b/Engine/Source/Runtime/Serialization/Private/Tests/StructSerializerTest.cpp
@@ -586,9 +586,9 @@ namespace StructSerializerTest
 				constexpr int32 TestIndex = 1;
 				FStructSerializer::SerializeElement(&TestStruct, Property, TestIndex, SerializerBackend, Policies);
 				Test.TestTrue(TEXT("Deserialization must succeed"), FStructDeserializer::DeserializeElement(&TestStruct2, *FStructSerializerArrayTestStruct::StaticStruct(), TestIndex, DeserializerBackend, DeserializerPolicies));
-				Test.TestNotEqual<int32>(TEXT("Arrays.StaticFloatArray[0] must not be the same before and after de-/serialization of element 1"), TestStruct.StaticFloatArray[0], TestStruct2.StaticFloatArray[0]);
-				Test.TestEqual<int32>(TEXT("Arrays.StaticFloatArray[1] must be the same before and after de-/serialization"), TestStruct.StaticFloatArray[TestIndex], TestStruct2.StaticFloatArray[TestIndex]);
-				Test.TestNotEqual<int32>(TEXT("Arrays.StaticFloatArray[2] must not be the same before and after de-/serialization of element 1"), TestStruct.StaticFloatArray[2], TestStruct2.StaticFloatArray[2]);
+				Test.TestNotEqual<float>(TEXT("Arrays.StaticFloatArray[0] must not be the same before and after de-/serialization of element 1"), TestStruct.StaticFloatArray[0], TestStruct2.StaticFloatArray[0]);
+				Test.TestEqual<float>(TEXT("Arrays.StaticFloatArray[1] must be the same before and after de-/serialization"), TestStruct.StaticFloatArray[TestIndex], TestStruct2.StaticFloatArray[TestIndex]);
+				Test.TestNotEqual<float>(TEXT("Arrays.StaticFloatArray[2] must not be the same before and after de-/serialization of element 1"), TestStruct.StaticFloatArray[2], TestStruct2.StaticFloatArray[2]);
 			}
 
 			//TMap<int32, FString> element
diff --git a/Engine/Source/Runtime/XmlParser/Private/FastXml.cpp b/Engine/Source/Runtime/XmlParser/Private/FastXml.cpp
--- a/Engine/Source/Runtime/XmlParser/Private/FastXml.cpp
+++ b/Engine/Source/Runtime/XmlParser/Private/FastXml.cpp
@@ -76,7 +76,7 @@ protected:
 	{
 		if( FeedbackContext != nullptr )
 		{
-			FeedbackContext->UpdateProgress( Buffer - &XmlFileContents[ 0 ], XmlFileContentsLength );
+			FeedbackContext->UpdateProgress(UE_PTRDIFF_TO_INT32(Buffer - &XmlFileContents[0]), IntCastChecked<int32>(XmlFileContentsLength) );
 		}
 
 		if( Char == L'/' || Char == L'?' )
