Added support for additional analytic-related argument parsing

diff --git a/Engine/Source/Runtime/Analytics/AnalyticsET/Private/IAnalyticsProviderET.cpp b/Engine/Source/Runtime/Analytics/AnalyticsET/Private/IAnalyticsProviderET.cpp
--- a/Engine/Source/Runtime/Analytics/AnalyticsET/Private/IAnalyticsProviderET.cpp
+++ b/Engine/Source/Runtime/Analytics/AnalyticsET/Private/IAnalyticsProviderET.cpp
@@ -78,6 +78,7 @@ public:
 	virtual void SetEventCallback(const OnEventRecorded& Callback) override;
 
 	virtual void SetURLEndpoint(const FString& UrlEndpoint, const TArray<FString>& AltDomains) override;
+	virtual void SetHeader(const FString& HeaderName, const FString& HeaderValue) override;
 	virtual void BlockUntilFlushed(float InTimeoutSec) override;
 	virtual void SetShouldRecordEventFunc(const ShouldRecordEventFunction& InShouldRecordEventFunc) override;
 	virtual ~FAnalyticsProviderET();
@@ -129,6 +130,8 @@ private:
 
 	TSharedPtr<class FHttpRetrySystem::FManager> HttpRetryManager;
 	FHttpRetrySystem::FRetryDomainsPtr RetryServers;
+	/** Http headers to add to requests */
+	TMap<FString, FString> HttpHeaders;
 };
 
 TSharedPtr<IAnalyticsProviderET> FAnalyticsET::CreateAnalyticsProvider(const Config& ConfigValues) const
@@ -321,6 +324,10 @@ TSharedRef<IHttpRequest, ESPMode::ThreadSafe> FAnalyticsProviderET::CreateReques
 		FHttpRetrySystem::FRetryResponseCodes(),
 		FHttpRetrySystem::FRetryVerbs(),
 		RetryServers);
+	for (const TPair<FString, FString>& HttpHeader : HttpHeaders)
+	{
+		HttpRequest->SetHeader(HttpHeader.Key, HttpHeader.Value);
+	}
 
 	return HttpRequest;
 }
@@ -599,6 +606,18 @@ void FAnalyticsProviderET::SetURLEndpoint(const FString& UrlEndpoint, const TArr
 	}
 }
 
+void FAnalyticsProviderET::SetHeader(const FString& HeaderName, const FString& HeaderValue)
+{
+	if (HeaderValue.IsEmpty())
+	{
+		HttpHeaders.Remove(HeaderName);
+	}
+	else
+	{
+		HttpHeaders.Emplace(HeaderName, HeaderValue);
+	}
+}
+
 void FAnalyticsProviderET::BlockUntilFlushed(float InTimeoutSec)
 {
 	FlushEvents();
diff --git a/Engine/Source/Runtime/Analytics/AnalyticsET/Public/IAnalyticsProviderET.h b/Engine/Source/Runtime/Analytics/AnalyticsET/Public/IAnalyticsProviderET.h
--- a/Engine/Source/Runtime/Analytics/AnalyticsET/Public/IAnalyticsProviderET.h
+++ b/Engine/Source/Runtime/Analytics/AnalyticsET/Public/IAnalyticsProviderET.h
@@ -123,6 +123,10 @@ public:
 	 * Updates the default URL endpoint and AltDomains.
 	 */
 	virtual void SetURLEndpoint(const FString& UrlEndpoint, const TArray<FString>& AltDomains) = 0;
+	/**
+	 * Sets a header to be included with analytics http requests
+	 */
+	virtual void SetHeader(const FString& HeaderName, const FString& HeaderValue) = 0;
 
 	typedef TFunction<void(const FString& EventName, const TArray<FAnalyticsEventAttribute>& Attrs, bool bJson)> OnEventRecorded;
 
