Removed Optimization Macros
Fixed Non-Unity compile issues

diff --git a/Engine/Source/Runtime/Analytics/AnalyticsET/Private/AnalyticsFlowTracker.cpp b/Engine/Source/Runtime/Analytics/AnalyticsET/Private/AnalyticsFlowTracker.cpp
--- a/Engine/Source/Runtime/Analytics/AnalyticsET/Private/AnalyticsFlowTracker.cpp
+++ b/Engine/Source/Runtime/Analytics/AnalyticsET/Private/AnalyticsFlowTracker.cpp
@@ -1,10 +1,7 @@
 // Copyright Epic Games, Inc. All Rights Reserved.
 
 #include "AnalyticsFlowTracker.h"
-
-#if UE_BUILD_DEVELOPMENT || UE_BUILD_DEBUG
-UE_DISABLE_OPTIMIZATION
-#endif
+#include "ProfilingDebugging/MiscTrace.h"
 
 void FAnalyticsFlowTracker::SetProvider(TSharedPtr<IAnalyticsProviderET>& InProvider)
 {
@@ -15,7 +12,6 @@ void FAnalyticsFlowTracker::StartSession()
 {
 }
 
-
 void FAnalyticsFlowTracker::EndSession()
 {
 	FScopeLock ScopeLock(&CriticalSection);
@@ -79,7 +75,7 @@ FGuid FAnalyticsFlowTracker::StartFlowStepInternal(const FName& NewFlowStepName,
 
 	if (ensureMsgf(FlowData, TEXT("FlowStep started outside of a valid flow scope")))
 	{
-		TRACE_BOOKMARK(TEXT("STARTFlowStep: %s"), *NewFlowStepName.ToString());
+		TRACE_BOOKMARK(TEXT("STARTFLOWSTEP: %s"), *NewFlowStepName.ToString());
 
 		// Create a new Guid for this FlowStep, can we assume it is unique?
 		FGuid NewFlowStepGuid = FGuid::NewGuid();
@@ -129,7 +125,7 @@ void FAnalyticsFlowTracker::EndFlowStepInternal(const FGuid& FlowStepGuid, bool
 			FlowStepData->EndTime = FDateTime::UtcNow();
 			FlowStepData->bSuccess = bSuccess;
 
-			TRACE_BOOKMARK(TEXT("ENDFlowStep: %s"), *FlowStepData->FlowStepName.ToString());
+			TRACE_BOOKMARK(TEXT("ENDFLOWSTEP: %s"), *FlowStepData->FlowStepName.ToString());
 
 			const FTimespan TimeTaken = FlowStepData->EndTime - FlowStepData->StartTime;
 			FlowStepData->TimeInSeconds = TimeTaken.GetTotalSeconds();
@@ -318,7 +314,3 @@ void FAnalyticsFlowTracker::EndFlowInternal(const FGuid& FlowGuid, bool bSuccess
 		}
 	}	
 }
-
-#if UE_BUILD_DEVELOPMENT || UE_BUILD_DEBUG
-UE_ENABLE_OPTIMIZATION
-#endif
