Renamed FlowStep back to SubFlow

diff --git a/Engine/Source/Runtime/Analytics/AnalyticsET/Private/AnalyticsFlowTracker.cpp b/Engine/Source/Runtime/Analytics/AnalyticsET/Private/AnalyticsFlowTracker.cpp
--- a/Engine/Source/Runtime/Analytics/AnalyticsET/Private/AnalyticsFlowTracker.cpp
+++ b/Engine/Source/Runtime/Analytics/AnalyticsET/Private/AnalyticsFlowTracker.cpp
@@ -51,100 +51,100 @@ FGuid FAnalyticsFlowTracker::StartFlow(const FName& NewFlowName)
 	return NewFlowGuid;
 }
 
-FGuid FAnalyticsFlowTracker::StartFlowStep(const FName& NewFlowStepName)
+FGuid FAnalyticsFlowTracker::StartSubFlow(const FName& NewSubFlowName)
 {
 	FScopeLock ScopeLock(&CriticalSection);
 	
 	if (FlowDataStack.Num()>0)
 	{
-		return StartFlowStepInternal(NewFlowStepName, FlowDataStack.Last(0));
+		return StartSubFlowInternal(NewSubFlowName, FlowDataStack.Last(0));
 	}
 
 	return FGuid();
 }
 
-FGuid FAnalyticsFlowTracker::StartFlowStep(const FName& NewFlowStepName, const FGuid& FlowGuid)
+FGuid FAnalyticsFlowTracker::StartSubFlow(const FName& NewSubFlowName, const FGuid& FlowGuid)
 {
 	FScopeLock ScopeLock(&CriticalSection);
-	return StartFlowStepInternal(NewFlowStepName, FlowGuid);
+	return StartSubFlowInternal(NewSubFlowName, FlowGuid);
 }
 
-FGuid FAnalyticsFlowTracker::StartFlowStepInternal(const FName& NewFlowStepName, const FGuid& FlowGuid)
+FGuid FAnalyticsFlowTracker::StartSubFlowInternal(const FName& NewSubFlowName, const FGuid& FlowGuid)
 {
 	FFlowData* FlowData = FlowDataRegistry.Find(FlowGuid);
 
-	if (ensureMsgf(FlowData, TEXT("FlowStep started outside of a valid flow scope")))
+	if (ensureMsgf(FlowData, TEXT("SubFlow started outside of a valid flow scope")))
 	{
-		TRACE_BOOKMARK(TEXT("STARTFLOWSTEP: %s"), *NewFlowStepName.ToString());
+		TRACE_BOOKMARK(TEXT("STARTSUBFLOW: %s"), *NewSubFlowName.ToString());
 
-		// Create a new Guid for this FlowStep, can we assume it is unique?
-		FGuid NewFlowStepGuid = FGuid::NewGuid();
-		ensureMsgf(FlowStepDataRegistry.Find(NewFlowStepGuid) == nullptr, TEXT("Could not generate a unique FlowStep guid."));
+		// Create a new Guid for this SubFlow, can we assume it is unique?
+		FGuid NewSubFlowGuid = FGuid::NewGuid();
+		ensureMsgf(SubFlowDataRegistry.Find(NewSubFlowGuid) == nullptr, TEXT("Could not generate a unique SubFlow guid."));
 
 		// Register the name and guid pair
-		FlowStepGuidRegistry.Add(NewFlowStepName, NewFlowStepGuid);
+		SubFlowGuidRegistry.Add(NewSubFlowName, NewSubFlowGuid);
 
-		FFlowStepData NewFlowStep;
+		FSubFlowData NewSubFlow;
 
-		NewFlowStep.FlowStepGuid = NewFlowStepGuid;
-		NewFlowStep.FlowStepName = NewFlowStepName;
-		NewFlowStep.StartTime = FDateTime::UtcNow();
-		NewFlowStep.EndTime = 0;
-		NewFlowStep.ScopeDepth = FlowData->FlowStepDataStack.Num();
+		NewSubFlow.SubFlowGuid = NewSubFlowGuid;
+		NewSubFlow.SubFlowName = NewSubFlowName;
+		NewSubFlow.StartTime = FDateTime::UtcNow();
+		NewSubFlow.EndTime = 0;
+		NewSubFlow.ScopeDepth = FlowData->SubFlowDataStack.Num();
 
-		// Add FlowStep to Flow
-		NewFlowStep.FlowGuid = FlowData->FlowGuid;
-		NewFlowStep.FlowName = FlowData->FlowName;
-		FlowData->FlowStepDataArray.Add(NewFlowStepGuid);
-		FlowData->FlowStepDataStack.Add(NewFlowStepGuid);
+		// Add SubFlow to Flow
+		NewSubFlow.FlowGuid = FlowData->FlowGuid;
+		NewSubFlow.FlowName = FlowData->FlowName;
+		FlowData->SubFlowDataArray.Add(NewSubFlowGuid);
+		FlowData->SubFlowDataStack.Add(NewSubFlowGuid);
 
 		// Register the name and guid pair
-		FlowStepGuidRegistry.Add(NewFlowStepName, NewFlowStepGuid);
-		FlowStepDataRegistry.Add(NewFlowStepGuid, NewFlowStep);
+		SubFlowGuidRegistry.Add(NewSubFlowName, NewSubFlowGuid);
+		SubFlowDataRegistry.Add(NewSubFlowGuid, NewSubFlow);
 
-		return NewFlowStepGuid;
+		return NewSubFlowGuid;
 	}
 
 	return FGuid();
 }
 
-void FAnalyticsFlowTracker::EndFlowStepInternal(const FGuid& FlowStepGuid, bool bSuccess, const TArray<FAnalyticsEventAttribute>& AdditionalAttributes)
+void FAnalyticsFlowTracker::EndSubFlowInternal(const FGuid& SubFlowGuid, bool bSuccess, const TArray<FAnalyticsEventAttribute>& AdditionalAttributes)
 {
-	if (FlowStepGuid.IsValid() == false)
+	if (SubFlowGuid.IsValid() == false)
 		return;
 
-	FFlowStepData* FlowStepData = FlowStepDataRegistry.Find(FlowStepGuid);
+	FSubFlowData* SubFlowData = SubFlowDataRegistry.Find(SubFlowGuid);
 
-	if (ensureMsgf(FlowStepData, TEXT("FlowStep does not exist.")))
+	if (ensureMsgf(SubFlowData, TEXT("SubFlow does not exist.")))
 	{
-		const FGuid FlowGuid = FlowStepData->FlowGuid;
+		const FGuid FlowGuid = SubFlowData->FlowGuid;
 
-		if (FlowStepData->TimeInSeconds == 0)
+		if (SubFlowData->TimeInSeconds == 0)
 		{
 			// Don't record again if it has already ended
-			FlowStepData->EndTime = FDateTime::UtcNow();
-			FlowStepData->bSuccess = bSuccess;
+			SubFlowData->EndTime = FDateTime::UtcNow();
+			SubFlowData->bSuccess = bSuccess;
 
-			TRACE_BOOKMARK(TEXT("ENDFLOWSTEP: %s"), *FlowStepData->FlowStepName.ToString());
+			TRACE_BOOKMARK(TEXT("ENDSUBFLOW: %s"), *SubFlowData->SubFlowName.ToString());
 
-			const FTimespan TimeTaken = FlowStepData->EndTime - FlowStepData->StartTime;
-			FlowStepData->TimeInSeconds = TimeTaken.GetTotalSeconds();
-			FlowStepData->AdditionalEventAttributes = AdditionalAttributes;
+			const FTimespan TimeTaken = SubFlowData->EndTime - SubFlowData->StartTime;
+			SubFlowData->TimeInSeconds = TimeTaken.GetTotalSeconds();
+			SubFlowData->AdditionalEventAttributes = AdditionalAttributes;
 
 			TArray<FAnalyticsEventAttribute> EventAttributes = AdditionalAttributes;
 
-			EventAttributes.Add(FAnalyticsEventAttribute(TEXT("SchemaVersion"), FlowStepSchemaVersion));
-			EventAttributes.Add(FAnalyticsEventAttribute(TEXT("FlowStepGUID"), *FlowStepData->FlowStepGuid.ToString()));
-			EventAttributes.Add(FAnalyticsEventAttribute(TEXT("FlowStepName"), *FlowStepData->FlowStepName.ToString()));
-			EventAttributes.Add(FAnalyticsEventAttribute(TEXT("FlowGUID"), FlowStepData->FlowGuid.ToString()));
-			EventAttributes.Add(FAnalyticsEventAttribute(TEXT("FlowName"), *FlowStepData->FlowName.ToString()));
-			EventAttributes.Add(FAnalyticsEventAttribute(TEXT("StartUTC"), FlowStepData->StartTime.ToUnixTimestampDecimal()));
-			EventAttributes.Add(FAnalyticsEventAttribute(TEXT("TimeInSec"), FlowStepData->TimeInSeconds));
-			EventAttributes.Add(FAnalyticsEventAttribute(TEXT("Success"), FlowStepData->bSuccess));
+			EventAttributes.Add(FAnalyticsEventAttribute(TEXT("SchemaVersion"), SubFlowSchemaVersion));
+			EventAttributes.Add(FAnalyticsEventAttribute(TEXT("SubFlowGUID"), *SubFlowData->SubFlowGuid.ToString()));
+			EventAttributes.Add(FAnalyticsEventAttribute(TEXT("SubFlowName"), *SubFlowData->SubFlowName.ToString()));
+			EventAttributes.Add(FAnalyticsEventAttribute(TEXT("FlowGUID"), SubFlowData->FlowGuid.ToString()));
+			EventAttributes.Add(FAnalyticsEventAttribute(TEXT("FlowName"), *SubFlowData->FlowName.ToString()));
+			EventAttributes.Add(FAnalyticsEventAttribute(TEXT("StartUTC"), SubFlowData->StartTime.ToUnixTimestampDecimal()));
+			EventAttributes.Add(FAnalyticsEventAttribute(TEXT("TimeInSec"), SubFlowData->TimeInSeconds));
+			EventAttributes.Add(FAnalyticsEventAttribute(TEXT("Success"), SubFlowData->bSuccess));
 
 			if (AnalyticsProvider.IsValid())
 			{
-				AnalyticsProvider->RecordEvent(FlowStepEventName, EventAttributes);
+				AnalyticsProvider->RecordEvent(SubFlowEventName, EventAttributes);
 			}
 
 			FFlowData* FlowData = FlowDataRegistry.Find(FlowGuid);
@@ -152,12 +152,12 @@ void FAnalyticsFlowTracker::EndFlowStepInternal(const FGuid& FlowStepGuid, bool
 			if (ensureMsgf(FlowData, TEXT("A sub flow does not belong to a valid flow.")))
 			{
 				// Most likely it will be the ending item
-				for (int32 index = FlowData->FlowStepDataStack.Num() - 1; index >= 0; index--)
+				for (int32 index = FlowData->SubFlowDataStack.Num() - 1; index >= 0; index--)
 				{
-					if (FlowData->FlowStepDataStack[index] == FlowStepGuid)
+					if (FlowData->SubFlowDataStack[index] == SubFlowGuid)
 					{
 						// Remove the sub flow from the stack
-						FlowData->FlowStepDataStack.RemoveAt(index);
+						FlowData->SubFlowDataStack.RemoveAt(index);
 						break;
 					}
 				}
@@ -166,22 +166,22 @@ void FAnalyticsFlowTracker::EndFlowStepInternal(const FGuid& FlowStepGuid, bool
 	}
 }
 
-void FAnalyticsFlowTracker::EndFlowStep(const FGuid& FlowStepGuid, bool bSuccess, const TArray<FAnalyticsEventAttribute>& AdditionalAttributes)
+void FAnalyticsFlowTracker::EndSubFlow(const FGuid& SubFlowGuid, bool bSuccess, const TArray<FAnalyticsEventAttribute>& AdditionalAttributes)
 {
 	FScopeLock ScopeLock(&CriticalSection);
-	if (FlowStepGuid.IsValid())
+	if (SubFlowGuid.IsValid())
 	{
-		EndFlowStepInternal(FlowStepGuid, bSuccess, AdditionalAttributes);
+		EndSubFlowInternal(SubFlowGuid, bSuccess, AdditionalAttributes);
 	}
 }
 
-void FAnalyticsFlowTracker::EndFlowStep(const FName& FlowStepName, bool bSuccess, const TArray<FAnalyticsEventAttribute>& AdditionalAttributes)
+void FAnalyticsFlowTracker::EndSubFlow(const FName& SubFlowName, bool bSuccess, const TArray<FAnalyticsEventAttribute>& AdditionalAttributes)
 {
 	FScopeLock ScopeLock(&CriticalSection);
-	FGuid* FlowStepGuid = FlowGuidRegistry.Find(FlowStepName);
-	if (FlowStepGuid)
+	FGuid* SubFlowGuid = FlowGuidRegistry.Find(SubFlowName);
+	if (SubFlowGuid)
 	{
-		EndFlowStepInternal(*FlowStepGuid, bSuccess, AdditionalAttributes);
+		EndSubFlowInternal(*SubFlowGuid, bSuccess, AdditionalAttributes);
 	}
 }
 
@@ -262,19 +262,19 @@ void FAnalyticsFlowTracker::EndFlowInternal(const FGuid& FlowGuid, bool bSuccess
 
 		double TotalTimeInSeconds = 0;
 
-		for (FGuid FlowStepGuid : FlowData->FlowStepDataArray)
+		for (FGuid SubFlowGuid : FlowData->SubFlowDataArray)
 		{
-			EndFlowStepInternal(FlowStepGuid);
+			EndSubFlowInternal(SubFlowGuid);
 
-			FFlowStepData* FlowStepData = FlowStepDataRegistry.Find(FlowStepGuid);
+			FSubFlowData* SubFlowData = SubFlowDataRegistry.Find(SubFlowGuid);
 
-			if (ensureMsgf(FlowStepData, TEXT("FlowStep does not exist.")))
+			if (ensureMsgf(SubFlowData, TEXT("SubFlow does not exist.")))
 			{
 				// Aggregate the additional attributes from the sub flows
-				AggregateAttributes(EventAttributes, FlowStepData->AdditionalEventAttributes);
+				AggregateAttributes(EventAttributes, SubFlowData->AdditionalEventAttributes);
 
-				TotalTimeInSeconds += FlowStepData->TimeInSeconds;
-				EventAttributes.Add(FAnalyticsEventAttribute(FlowStepData->FlowStepName.ToString(), FlowStepData->TimeInSeconds));
+				TotalTimeInSeconds += SubFlowData->TimeInSeconds;
+				EventAttributes.Add(FAnalyticsEventAttribute(SubFlowData->SubFlowName.ToString(), SubFlowData->TimeInSeconds));
 			}
 		}
 
@@ -286,15 +286,15 @@ void FAnalyticsFlowTracker::EndFlowInternal(const FGuid& FlowGuid, bool bSuccess
 		}
 
 		// Clear up our data
-		for (FGuid FlowStepGuid : FlowData->FlowStepDataArray)
+		for (FGuid SubFlowGuid : FlowData->SubFlowDataArray)
 		{
-			FFlowStepData* FlowStepData = FlowStepDataRegistry.Find(FlowStepGuid);
+			FSubFlowData* SubFlowData = SubFlowDataRegistry.Find(SubFlowGuid);
 
-			if (ensureMsgf(FlowStepData, TEXT("FlowStep does not exist.")))
+			if (ensureMsgf(SubFlowData, TEXT("SubFlow does not exist.")))
 			{
-				// Remove the FlowStep and guid from the registry
-				FlowStepDataRegistry.Remove(FlowStepGuid);
-				FlowStepGuidRegistry.Remove(FlowStepData->FlowStepName);
+				// Remove the SubFlow and guid from the registry
+				SubFlowDataRegistry.Remove(SubFlowGuid);
+				SubFlowGuidRegistry.Remove(SubFlowData->SubFlowName);
 			}
 		}
 
diff --git a/Engine/Source/Runtime/Analytics/AnalyticsET/Public/AnalyticsFlowTracker.h b/Engine/Source/Runtime/Analytics/AnalyticsET/Public/AnalyticsFlowTracker.h
--- a/Engine/Source/Runtime/Analytics/AnalyticsET/Public/AnalyticsFlowTracker.h
+++ b/Engine/Source/Runtime/Analytics/AnalyticsET/Public/AnalyticsFlowTracker.h
@@ -15,10 +15,10 @@ public:
 	/** Sets the analytics provider for the flow tracker. */
 	ANALYTICSET_API void SetProvider(TSharedPtr<IAnalyticsProviderET>& AnalyticsProvider);
 
-	/** Begins a new flow tracking session. Will emit Flow and FlowStep events to the specified analytics provider */
+	/** Begins a new flow tracking session. Will emit Flow and SubFlow events to the specified analytics provider */
 	ANALYTICSET_API void StartSession();
 
-	/** Ends all open Flows and FlowSteps*/
+	/** Ends all open Flows and SubFlows*/
 	ANALYTICSET_API void EndSession();
 	
 	/** Start a new Flow, the existing flow context will be pushed onto a stack and the new flow will become the current context*/
@@ -34,25 +34,25 @@ public:
 	ANALYTICSET_API void EndFlow(const FGuid& FlowGuid, bool bSuccess = true, const TArray<FAnalyticsEventAttribute>& AdditionalAttributes = {});
 
 	/** Start a new flow and add it to the current flow context */
-	ANALYTICSET_API FGuid StartFlowStep(const FName& FlowStepName);
+	ANALYTICSET_API FGuid StartSubFlow(const FName& SubFlowName);
 
 	/** Start a new flow step and add it to a specific flow context by GUID */
-	ANALYTICSET_API FGuid StartFlowStep(const FName& FlowStepName, const FGuid& FlowGuid);
+	ANALYTICSET_API FGuid StartSubFlow(const FName& SubFlowName, const FGuid& FlowGuid);
 
 	/** End an existing flow step by name */
-	ANALYTICSET_API void EndFlowStep(const FName& FlowStepName, bool bSuccess = true, const TArray<FAnalyticsEventAttribute>& AdditionalAttributes = {});
+	ANALYTICSET_API void EndSubFlow(const FName& SubFlowName, bool bSuccess = true, const TArray<FAnalyticsEventAttribute>& AdditionalAttributes = {});
 
 	/** End an existing flow step by GUID */
-	ANALYTICSET_API void EndFlowStep(const FGuid& FlowStepGuid, bool bSuccess = true, const TArray<FAnalyticsEventAttribute>& AdditionalAttributes = {});
+	ANALYTICSET_API void EndSubFlow(const FGuid& SubFlowGuid, bool bSuccess = true, const TArray<FAnalyticsEventAttribute>& AdditionalAttributes = {});
 
 private:
 
-	struct FFlowStepData
+	struct FSubFlowData
 	{
 		FName		FlowName;
 		FGuid		FlowGuid = FGuid();
-		FName		FlowStepName;
-		FGuid		FlowStepGuid;
+		FName		SubFlowName;
+		FGuid		SubFlowGuid;
 		FDateTime	StartTime = 0;
 		FDateTime	EndTime = 0;
 		double		TimeInSeconds = 0;
@@ -69,19 +69,19 @@ private:
 		FDateTime					StartTime = 0;
 		FDateTime					EndTime = 0;
 		double						TimeInSeconds = 0;
-		TArray<FGuid>				FlowStepDataArray;
-		TArray<FGuid>				FlowStepDataStack;
+		TArray<FGuid>				SubFlowDataArray;
+		TArray<FGuid>				SubFlowDataStack;
 	};
 
 	void EndFlowInternal(const FGuid& FlowGuid, bool bSuccess = true, const TArray<FAnalyticsEventAttribute>& AdditionalAttributes = {});
 
-	FGuid StartFlowStepInternal(const FName& NewFlowStepName, const FGuid& FlowGuid);
-	void EndFlowStepInternal(const FGuid& FlowStepGuid, bool bSuccess = true, const TArray<FAnalyticsEventAttribute>& AdditionalAttributes = {});
+	FGuid StartSubFlowInternal(const FName& NewSubFlowName, const FGuid& FlowGuid);
+	void EndSubFlowInternal(const FGuid& SubFlowGuid, bool bSuccess = true, const TArray<FAnalyticsEventAttribute>& AdditionalAttributes = {});
 
 	TMap<FName, FGuid>			FlowGuidRegistry;
 	TMap<FGuid, FFlowData>		FlowDataRegistry;
-	TMap<FName, FGuid>			FlowStepGuidRegistry;
-	TMap<FGuid, FFlowStepData>	FlowStepDataRegistry;
+	TMap<FName, FGuid>			SubFlowGuidRegistry;
+	TMap<FGuid, FSubFlowData>	SubFlowDataRegistry;
 	TArray<FGuid>				FlowDataStack;
 	FCriticalSection			CriticalSection;
 	TSharedPtr<IAnalyticsProviderET> AnalyticsProvider;
@@ -89,8 +89,8 @@ private:
 	const uint32 FlowSchemaVersion = 4;
 	const FString FlowEventName = TEXT("Iteration.Flow");
 
-	const uint32 FlowStepSchemaVersion = 4;
-	const FString FlowStepEventName = TEXT("Iteration.FlowStep");
+	const uint32 SubFlowSchemaVersion = 4;
+	const FString SubFlowEventName = TEXT("Iteration.SubFlow");
 };
 
 
