Improved Editor Flow Analytics and extended the EditorAnalytics to demonstrate correct usage of SubFLow/Flow
Added correct PIE Flow and SubFlow tracking
Added WorldStreaming tracking
Fixed several bugs in Flow Tracker
Added ThreadId to SubFlow and FLow types

diff --git a/Engine/Source/Runtime/Analytics/AnalyticsET/Private/AnalyticsFlowTracker.cpp b/Engine/Source/Runtime/Analytics/AnalyticsET/Private/AnalyticsFlowTracker.cpp
--- a/Engine/Source/Runtime/Analytics/AnalyticsET/Private/AnalyticsFlowTracker.cpp
+++ b/Engine/Source/Runtime/Analytics/AnalyticsET/Private/AnalyticsFlowTracker.cpp
@@ -42,6 +42,7 @@ FGuid FAnalyticsFlowTracker::StartFlow(const FName& NewFlowName)
 	FlowData.StartTime = FDateTime::UtcNow();
 	FlowData.FlowName = NewFlowName;
 	FlowData.FlowGuid = NewFlowGuid;
+	FlowData.ThreadId = FPlatformTLS::GetCurrentThreadId();
 
 	// Register the name and guid pair
 	FlowGuidRegistry.Add(NewFlowName, NewFlowGuid);
@@ -91,6 +92,7 @@ FGuid FAnalyticsFlowTracker::StartSubFlowInternal(const FName& NewSubFlowName, c
 		NewSubFlow.StartTime = FDateTime::UtcNow();
 		NewSubFlow.EndTime = 0;
 		NewSubFlow.ScopeDepth = FlowData->SubFlowDataStack.Num();
+		NewSubFlow.ThreadId = FPlatformTLS::GetCurrentThreadId();
 
 		// Add SubFlow to Flow
 		NewSubFlow.FlowGuid = FlowData->FlowGuid;
@@ -108,81 +110,88 @@ FGuid FAnalyticsFlowTracker::StartSubFlowInternal(const FName& NewSubFlowName, c
 	return FGuid();
 }
 
-void FAnalyticsFlowTracker::EndSubFlowInternal(const FGuid& SubFlowGuid, bool bSuccess, const TArray<FAnalyticsEventAttribute>& AdditionalAttributes)
+bool FAnalyticsFlowTracker::EndSubFlowInternal(const FGuid& SubFlowGuid, bool bSuccess, const TArray<FAnalyticsEventAttribute>& AdditionalAttributes)
 {
 	if (SubFlowGuid.IsValid() == false)
-		return;
+		return false;
 
 	FSubFlowData* SubFlowData = SubFlowDataRegistry.Find(SubFlowGuid);
 
-	if (ensureMsgf(SubFlowData, TEXT("SubFlow does not exist.")))
-	{
-		const FGuid FlowGuid = SubFlowData->FlowGuid;
+	if (SubFlowData == nullptr)
+		return false;
 
-		if (SubFlowData->TimeInSeconds == 0)
-		{
-			// Don't record again if it has already ended
-			SubFlowData->EndTime = FDateTime::UtcNow();
-			SubFlowData->bSuccess = bSuccess;
+	const FGuid FlowGuid = SubFlowData->FlowGuid;
 
-			TRACE_BOOKMARK(TEXT("ENDSUBFLOW: %s"), *SubFlowData->SubFlowName.ToString());
+	if (SubFlowData->EndTime<SubFlowData->StartTime)
+	{
+		// Don't record again if it has already ended
+		SubFlowData->EndTime = FDateTime::UtcNow();
+		SubFlowData->bSuccess = bSuccess;
 
-			const FTimespan TimeTaken = SubFlowData->EndTime - SubFlowData->StartTime;
-			SubFlowData->TimeInSeconds = TimeTaken.GetTotalSeconds();
-			SubFlowData->AdditionalEventAttributes = AdditionalAttributes;
+		const FTimespan TimeTaken = SubFlowData->EndTime - SubFlowData->StartTime;
+		SubFlowData->TimeInSeconds = TimeTaken.GetTotalSeconds();
+		SubFlowData->AdditionalEventAttributes = AdditionalAttributes;
 
-			TArray<FAnalyticsEventAttribute> EventAttributes = AdditionalAttributes;
+		TArray<FAnalyticsEventAttribute> EventAttributes = AdditionalAttributes;
 
-			EventAttributes.Add(FAnalyticsEventAttribute(TEXT("SchemaVersion"), SubFlowSchemaVersion));
-			EventAttributes.Add(FAnalyticsEventAttribute(TEXT("SubFlowGUID"), *SubFlowData->SubFlowGuid.ToString()));
-			EventAttributes.Add(FAnalyticsEventAttribute(TEXT("SubFlowName"), *SubFlowData->SubFlowName.ToString()));
-			EventAttributes.Add(FAnalyticsEventAttribute(TEXT("FlowGUID"), SubFlowData->FlowGuid.ToString()));
-			EventAttributes.Add(FAnalyticsEventAttribute(TEXT("FlowName"), *SubFlowData->FlowName.ToString()));
-			EventAttributes.Add(FAnalyticsEventAttribute(TEXT("StartUTC"), SubFlowData->StartTime.ToUnixTimestampDecimal()));
-			EventAttributes.Add(FAnalyticsEventAttribute(TEXT("TimeInSec"), SubFlowData->TimeInSeconds));
-			EventAttributes.Add(FAnalyticsEventAttribute(TEXT("Success"), SubFlowData->bSuccess));
+		EventAttributes.Add(FAnalyticsEventAttribute(TEXT("SchemaVersion"), SubFlowSchemaVersion));
+		EventAttributes.Add(FAnalyticsEventAttribute(TEXT("SubFlowGUID"), *SubFlowData->SubFlowGuid.ToString()));
+		EventAttributes.Add(FAnalyticsEventAttribute(TEXT("SubFlowName"), *SubFlowData->SubFlowName.ToString()));
+		EventAttributes.Add(FAnalyticsEventAttribute(TEXT("FlowGUID"), SubFlowData->FlowGuid.ToString()));
+		EventAttributes.Add(FAnalyticsEventAttribute(TEXT("FlowName"), *SubFlowData->FlowName.ToString()));
+		EventAttributes.Add(FAnalyticsEventAttribute(TEXT("ThreadId"), SubFlowData->ThreadId));
+		EventAttributes.Add(FAnalyticsEventAttribute(TEXT("StartUTC"), SubFlowData->StartTime.ToUnixTimestampDecimal()));
+		EventAttributes.Add(FAnalyticsEventAttribute(TEXT("TimeInSec"), SubFlowData->TimeInSeconds));
+		EventAttributes.Add(FAnalyticsEventAttribute(TEXT("Success"), SubFlowData->bSuccess));
 
-			if (AnalyticsProvider.IsValid())
-			{
-				AnalyticsProvider->RecordEvent(SubFlowEventName, EventAttributes);
-			}
+		if (AnalyticsProvider.IsValid())
+		{
+			AnalyticsProvider->RecordEvent(SubFlowEventName, EventAttributes);
+		}
 
-			FFlowData* FlowData = FlowDataRegistry.Find(FlowGuid);
+		TRACE_BOOKMARK(TEXT("ENDSUBFLOW: %s"), *SubFlowData->SubFlowName.ToString());
 
-			if (ensureMsgf(FlowData, TEXT("A sub flow does not belong to a valid flow.")))
+		FFlowData* FlowData = FlowDataRegistry.Find(FlowGuid);
+
+		if (ensureMsgf(FlowData, TEXT("A sub flow does not belong to a valid flow.")))
+		{
+			// Most likely it will be the ending item
+			for (int32 index = FlowData->SubFlowDataStack.Num() - 1; index >= 0; index--)
 			{
-				// Most likely it will be the ending item
-				for (int32 index = FlowData->SubFlowDataStack.Num() - 1; index >= 0; index--)
+				if (FlowData->SubFlowDataStack[index] == SubFlowGuid)
 				{
-					if (FlowData->SubFlowDataStack[index] == SubFlowGuid)
-					{
-						// Remove the sub flow from the stack
-						FlowData->SubFlowDataStack.RemoveAt(index);
-						break;
-					}
+					// Remove the sub flow from the stack
+					FlowData->SubFlowDataStack.RemoveAt(index);
+					break;
 				}
 			}
 		}
 	}
+	
+	return true;
 }
 
-void FAnalyticsFlowTracker::EndSubFlow(const FGuid& SubFlowGuid, bool bSuccess, const TArray<FAnalyticsEventAttribute>& AdditionalAttributes)
+bool FAnalyticsFlowTracker::EndSubFlow(const FGuid& SubFlowGuid, bool bSuccess, const TArray<FAnalyticsEventAttribute>& AdditionalAttributes)
 {
 	FScopeLock ScopeLock(&CriticalSection);
 	if (SubFlowGuid.IsValid())
 	{
-		EndSubFlowInternal(SubFlowGuid, bSuccess, AdditionalAttributes);
+		return EndSubFlowInternal(SubFlowGuid, bSuccess, AdditionalAttributes);
 	}
+
+	return false;
 }
 
-void FAnalyticsFlowTracker::EndSubFlow(const FName& SubFlowName, bool bSuccess, const TArray<FAnalyticsEventAttribute>& AdditionalAttributes)
+bool FAnalyticsFlowTracker::EndSubFlow(const FName& SubFlowName, bool bSuccess, const TArray<FAnalyticsEventAttribute>& AdditionalAttributes)
 {
 	FScopeLock ScopeLock(&CriticalSection);
-	FGuid* SubFlowGuid = FlowGuidRegistry.Find(SubFlowName);
+	FGuid* SubFlowGuid = SubFlowGuidRegistry.Find(SubFlowName);
 	if (SubFlowGuid)
 	{
-		EndSubFlowInternal(*SubFlowGuid, bSuccess, AdditionalAttributes);
+		return EndSubFlowInternal(*SubFlowGuid, bSuccess, AdditionalAttributes);
 	}
+
+	return false;
 }
 
 static void AggregateAttributes(TArray<FAnalyticsEventAttribute>& AggregatedAttibutes, const TArray<FAnalyticsEventAttribute>& Attributes)
@@ -212,105 +221,113 @@ static void AggregateAttributes(TArray<FAnalyticsEventAttribute>& AggregatedAtti
 	}
 }
 
-void FAnalyticsFlowTracker::EndFlow(const FName& FlowName, bool bSuccess, const TArray<FAnalyticsEventAttribute>& AdditionalAttributes)
+bool FAnalyticsFlowTracker::EndFlow(const FName& FlowName, bool bSuccess, const TArray<FAnalyticsEventAttribute>& AdditionalAttributes)
 {
 	FScopeLock ScopeLock(&CriticalSection);
 	if (FGuid* FlowGuid = FlowGuidRegistry.Find(FlowName))
 	{
-		EndFlowInternal(*FlowGuid, bSuccess, AdditionalAttributes);
+		return EndFlowInternal(*FlowGuid, bSuccess, AdditionalAttributes);
 	}
+
+	return false;
 }
 
-void FAnalyticsFlowTracker::EndFlow(bool bSuccess, const TArray<FAnalyticsEventAttribute>& AdditionalAttributes)
+bool FAnalyticsFlowTracker::EndFlow(bool bSuccess, const TArray<FAnalyticsEventAttribute>& AdditionalAttributes)
 {
 	FScopeLock ScopeLock(&CriticalSection);
 
 	if (FlowDataStack.IsEmpty() == false)
 	{
-		EndFlowInternal(FlowDataStack.Last(0), bSuccess, AdditionalAttributes);
+		return EndFlowInternal(FlowDataStack.Last(0), bSuccess, AdditionalAttributes);
 	}
+
+	return false;
 }
 
-void FAnalyticsFlowTracker::EndFlow(const FGuid& FlowGuid, bool bSuccess, const TArray<FAnalyticsEventAttribute>& AdditionalAttributes)
+bool FAnalyticsFlowTracker::EndFlow(const FGuid& FlowGuid, bool bSuccess, const TArray<FAnalyticsEventAttribute>& AdditionalAttributes)
 {
 	FScopeLock ScopeLock(&CriticalSection);
-	EndFlowInternal(FlowGuid, bSuccess, AdditionalAttributes);
+	return EndFlowInternal(FlowGuid, bSuccess, AdditionalAttributes);
 }
 
-void FAnalyticsFlowTracker::EndFlowInternal(const FGuid& FlowGuid, bool bSuccess, const TArray<FAnalyticsEventAttribute>& AdditionalAttributes)
+bool FAnalyticsFlowTracker::EndFlowInternal(const FGuid& FlowGuid, bool bSuccess, const TArray<FAnalyticsEventAttribute>& AdditionalAttributes)
 {
 	if (FlowGuid.IsValid() == false)
-		return;
+		return false;
 
 	FFlowData* FlowData = FlowDataRegistry.Find(FlowGuid);
+	
+	if (FlowData == nullptr)
+		return false;
 
-	if (ensureMsgf(FlowData, TEXT("There is no valid flow")))
-	{
-		FlowData->EndTime = FDateTime::UtcNow();
-		const FTimespan WallTime = FlowData->EndTime - FlowData->StartTime;
-		FlowData->TimeInSeconds = WallTime.GetTotalSeconds();
+	FlowData->EndTime = FDateTime::UtcNow();
+	const FTimespan WallTime = FlowData->EndTime - FlowData->StartTime;
+	FlowData->TimeInSeconds = WallTime.GetTotalSeconds();
 
-		TRACE_BOOKMARK(TEXT("ENDFLOW: %s"), *FlowData->FlowName.ToString());
+	TArray<FAnalyticsEventAttribute> EventAttributes = AdditionalAttributes;
+	EventAttributes.Add(FAnalyticsEventAttribute(TEXT("SchemaVersion"), FlowSchemaVersion));
+	EventAttributes.Add(FAnalyticsEventAttribute(TEXT("FlowGUID"), FlowData->FlowGuid.ToString()));
+	EventAttributes.Add(FAnalyticsEventAttribute(TEXT("FlowName"), FlowData->FlowName.ToString()));
+	EventAttributes.Add(FAnalyticsEventAttribute(TEXT("ThreadId"), FlowData->ThreadId));
+	EventAttributes.Add(FAnalyticsEventAttribute(TEXT("StartUTC"), FlowData->StartTime.ToUnixTimestampDecimal()));
+	EventAttributes.Add(FAnalyticsEventAttribute(TEXT("Success"), bSuccess));
+	EventAttributes.Add(FAnalyticsEventAttribute(TEXT("WallTimeInSec"), FlowData->TimeInSeconds));
 
-		TArray<FAnalyticsEventAttribute> EventAttributes = AdditionalAttributes;
-		EventAttributes.Add(FAnalyticsEventAttribute(TEXT("SchemaVersion"), FlowSchemaVersion));
-		EventAttributes.Add(FAnalyticsEventAttribute(TEXT("FlowGUID"), FlowData->FlowGuid.ToString()));
-		EventAttributes.Add(FAnalyticsEventAttribute(TEXT("FlowName"), FlowData->FlowName.ToString()));
-		EventAttributes.Add(FAnalyticsEventAttribute(TEXT("StartUTC"), FlowData->StartTime.ToUnixTimestampDecimal()));
-		EventAttributes.Add(FAnalyticsEventAttribute(TEXT("Success"), bSuccess));
-		EventAttributes.Add(FAnalyticsEventAttribute(TEXT("WallTimeInSec"), FlowData->TimeInSeconds));
+	double TotalTimeInSeconds = 0;
+
+	for (FGuid SubFlowGuid : FlowData->SubFlowDataArray)
+	{
+		EndSubFlowInternal(SubFlowGuid);
 
-		double TotalTimeInSeconds = 0;
+		FSubFlowData* SubFlowData = SubFlowDataRegistry.Find(SubFlowGuid);
 
-		for (FGuid SubFlowGuid : FlowData->SubFlowDataArray)
+		if (ensureMsgf(SubFlowData, TEXT("SubFlow does not exist.")))
 		{
-			EndSubFlowInternal(SubFlowGuid);
+			// Aggregate the additional attributes from the sub flows
+			AggregateAttributes(EventAttributes, SubFlowData->AdditionalEventAttributes);
 
-			FSubFlowData* SubFlowData = SubFlowDataRegistry.Find(SubFlowGuid);
+			TotalTimeInSeconds += SubFlowData->TimeInSeconds;
+			EventAttributes.Add(FAnalyticsEventAttribute(SubFlowData->SubFlowName.ToString(), SubFlowData->TimeInSeconds));
+		}
+	}
 
-			if (ensureMsgf(SubFlowData, TEXT("SubFlow does not exist.")))
-			{
-				// Aggregate the additional attributes from the sub flows
-				AggregateAttributes(EventAttributes, SubFlowData->AdditionalEventAttributes);
+	EventAttributes.Add(FAnalyticsEventAttribute(TEXT("TotalTimeInSec"), TotalTimeInSeconds));
 
-				TotalTimeInSeconds += SubFlowData->TimeInSeconds;
-				EventAttributes.Add(FAnalyticsEventAttribute(SubFlowData->SubFlowName.ToString(), SubFlowData->TimeInSeconds));
-			}
-		}
+	if (AnalyticsProvider.IsValid())
+	{
+		AnalyticsProvider->RecordEvent(FlowEventName, EventAttributes);
+		AnalyticsProvider->FlushEvents();
+	}
 
-		EventAttributes.Add(FAnalyticsEventAttribute(TEXT("TotalTimeInSec"), TotalTimeInSeconds));
+	TRACE_BOOKMARK(TEXT("ENDFLOW: %s"), *FlowData->FlowName.ToString());
 
-		if (AnalyticsProvider.IsValid())
-		{
-			AnalyticsProvider->RecordEvent(FlowEventName, EventAttributes);
-		}
+	// Clear up our data
+	for (FGuid SubFlowGuid : FlowData->SubFlowDataArray)
+	{
+		FSubFlowData* SubFlowData = SubFlowDataRegistry.Find(SubFlowGuid);
 
-		// Clear up our data
-		for (FGuid SubFlowGuid : FlowData->SubFlowDataArray)
+		if (ensureMsgf(SubFlowData, TEXT("SubFlow does not exist.")))
 		{
-			FSubFlowData* SubFlowData = SubFlowDataRegistry.Find(SubFlowGuid);
-
-			if (ensureMsgf(SubFlowData, TEXT("SubFlow does not exist.")))
-			{
-				// Remove the SubFlow and guid from the registry
-				SubFlowDataRegistry.Remove(SubFlowGuid);
-				SubFlowGuidRegistry.Remove(SubFlowData->SubFlowName);
-			}
+			// Remove the SubFlow and guid from the registry
+			SubFlowDataRegistry.Remove(SubFlowGuid);
+			SubFlowGuidRegistry.Remove(SubFlowData->SubFlowName);
 		}
+	}
 
-		// Remove the flow and guid from the registry
-		FlowDataRegistry.Remove(FlowData->FlowGuid);
-		FlowGuidRegistry.Remove(FlowData->FlowName);
+	// Remove the flow and guid from the registry
+	FlowDataRegistry.Remove(FlowData->FlowGuid);
+	FlowGuidRegistry.Remove(FlowData->FlowName);
 
-		// Remove the FlowData from the stack
-		for (int32 index = FlowDataStack.Num() - 1; index >= 0; index--)
+	// Remove the FlowData from the stack
+	for (int32 index = FlowDataStack.Num() - 1; index >= 0; index--)
+	{
+		if (FlowDataStack[index] == FlowGuid)
 		{
-			if (FlowDataStack[index] == FlowGuid)
-			{
-				// Remove the flow from the stack
-				FlowDataStack.RemoveAt(index);
-				break;
-			}
+			// Remove the flow from the stack
+			FlowDataStack.RemoveAt(index);
+			break;
 		}
-	}	
+	}
+
+	return true;
 }
diff --git a/Engine/Source/Runtime/Analytics/AnalyticsET/Public/AnalyticsFlowTracker.h b/Engine/Source/Runtime/Analytics/AnalyticsET/Public/AnalyticsFlowTracker.h
--- a/Engine/Source/Runtime/Analytics/AnalyticsET/Public/AnalyticsFlowTracker.h
+++ b/Engine/Source/Runtime/Analytics/AnalyticsET/Public/AnalyticsFlowTracker.h
@@ -6,6 +6,8 @@
 #include "UObject/UnrealNames.h"
 #include "IAnalyticsProviderET.h"
 
+using FThreadId = uint32;
+
 class FAnalyticsFlowTracker : FNoncopyable
 {
 public:
@@ -25,13 +27,13 @@ public:
 	ANALYTICSET_API FGuid StartFlow(const FName& FlowName);
 
 	/** End the flow for the current context and pop the stack*/
-	ANALYTICSET_API void EndFlow(bool bSuccess = true, const TArray<FAnalyticsEventAttribute>& AdditionalAttributes = {});
+	ANALYTICSET_API bool EndFlow(bool bSuccess = true, const TArray<FAnalyticsEventAttribute>& AdditionalAttributes = {});
 
 	/** End an existing flow by name */
-	ANALYTICSET_API void EndFlow(const FName& FlowName, bool bSuccess = true, const TArray<FAnalyticsEventAttribute>& AdditionalAttributes = {});
+	ANALYTICSET_API bool EndFlow(const FName& FlowName, bool bSuccess = true, const TArray<FAnalyticsEventAttribute>& AdditionalAttributes = {});
 
-	/** End an existing flow by GUID */
-	ANALYTICSET_API void EndFlow(const FGuid& FlowGuid, bool bSuccess = true, const TArray<FAnalyticsEventAttribute>& AdditionalAttributes = {});
+	/** End an bool flow by GUID */
+	ANALYTICSET_API bool EndFlow(const FGuid& FlowGuid, bool bSuccess = true, const TArray<FAnalyticsEventAttribute>& AdditionalAttributes = {});
 
 	/** Start a new flow and add it to the current flow context */
 	ANALYTICSET_API FGuid StartSubFlow(const FName& SubFlowName);
@@ -40,10 +42,10 @@ public:
 	ANALYTICSET_API FGuid StartSubFlow(const FName& SubFlowName, const FGuid& FlowGuid);
 
 	/** End an existing flow step by name */
-	ANALYTICSET_API void EndSubFlow(const FName& SubFlowName, bool bSuccess = true, const TArray<FAnalyticsEventAttribute>& AdditionalAttributes = {});
+	ANALYTICSET_API bool EndSubFlow(const FName& SubFlowName, bool bSuccess = true, const TArray<FAnalyticsEventAttribute>& AdditionalAttributes = {});
 
 	/** End an existing flow step by GUID */
-	ANALYTICSET_API void EndSubFlow(const FGuid& SubFlowGuid, bool bSuccess = true, const TArray<FAnalyticsEventAttribute>& AdditionalAttributes = {});
+	ANALYTICSET_API bool EndSubFlow(const FGuid& SubFlowGuid, bool bSuccess = true, const TArray<FAnalyticsEventAttribute>& AdditionalAttributes = {});
 
 private:
 
@@ -55,6 +57,7 @@ private:
 		FGuid		SubFlowGuid;
 		FDateTime	StartTime = 0;
 		FDateTime	EndTime = 0;
+		FThreadId	ThreadId = 0;
 		double		TimeInSeconds = 0;
 		bool		bSuccess = false;
 		int32		ScopeDepth;
@@ -68,22 +71,23 @@ private:
 		FGuid						FlowGuid = FGuid();
 		FDateTime					StartTime = 0;
 		FDateTime					EndTime = 0;
+		FThreadId					ThreadId = 0 ;
 		double						TimeInSeconds = 0;
 		TArray<FGuid>				SubFlowDataArray;
 		TArray<FGuid>				SubFlowDataStack;
 	};
 
-	void EndFlowInternal(const FGuid& FlowGuid, bool bSuccess = true, const TArray<FAnalyticsEventAttribute>& AdditionalAttributes = {});
+	bool EndFlowInternal(const FGuid& FlowGuid, bool bSuccess = true, const TArray<FAnalyticsEventAttribute>& AdditionalAttributes = {});
 
 	FGuid StartSubFlowInternal(const FName& NewSubFlowName, const FGuid& FlowGuid);
-	void EndSubFlowInternal(const FGuid& SubFlowGuid, bool bSuccess = true, const TArray<FAnalyticsEventAttribute>& AdditionalAttributes = {});
-
-	TMap<FName, FGuid>			FlowGuidRegistry;
-	TMap<FGuid, FFlowData>		FlowDataRegistry;
-	TMap<FName, FGuid>			SubFlowGuidRegistry;
-	TMap<FGuid, FSubFlowData>	SubFlowDataRegistry;
-	TArray<FGuid>				FlowDataStack;
-	FCriticalSection			CriticalSection;
+	bool EndSubFlowInternal(const FGuid& SubFlowGuid, bool bSuccess = true, const TArray<FAnalyticsEventAttribute>& AdditionalAttributes = {});
+
+	TMap<FName, FGuid>				FlowGuidRegistry;
+	TMap<FGuid, FFlowData>			FlowDataRegistry;
+	TMap<FName, FGuid>				SubFlowGuidRegistry;
+	TMap<FGuid, FSubFlowData>		SubFlowDataRegistry;
+	TArray<FGuid>					FlowDataStack;
+	FCriticalSection				CriticalSection;
 	TSharedPtr<IAnalyticsProviderET> AnalyticsProvider;
 
 	const uint32 FlowSchemaVersion = 4;
