Collection Manager Telemetry Events
* Uses method of exposing a delegate for telemetry events to be registered by a StudioAnalytic provider implementation to avoid embedding schema in the engine
* Generic templates for RAII scoped timed telemetry events added to aid in generating timed events

diff --git a/Engine/Source/Developer/CollectionManager/CollectionManager.Build.cs b/Engine/Source/Developer/CollectionManager/CollectionManager.Build.cs
--- a/Engine/Source/Developer/CollectionManager/CollectionManager.Build.cs
+++ b/Engine/Source/Developer/CollectionManager/CollectionManager.Build.cs
@@ -13,6 +13,7 @@ public class CollectionManager : ModuleRules
 				"Core",
 				"CoreUObject",
 				"DirectoryWatcher",
+				"Analytics",
 				"SourceControl"
 			}
 			);
diff --git a/Engine/Source/Developer/CollectionManager/Private/Collection.h b/Engine/Source/Developer/CollectionManager/Private/Collection.h
--- a/Engine/Source/Developer/CollectionManager/Private/Collection.h
+++ b/Engine/Source/Developer/CollectionManager/Private/Collection.h
@@ -91,6 +91,9 @@ public:
 	/** Whether the collection has any contents */
 	bool IsEmpty() const;
 
+	/** Whether the Collection uses Source Control */
+	bool IsUsingSourceControl() const;
+
 	/** Logs the contents of the collection */
 	void PrintCollection() const;
 
@@ -212,3 +215,8 @@ private:
 	/** The state of the collection the last time it was loaded from or saved to disk. */
 	FCollectionSnapshot DiskSnapshot;
 };
+
+inline bool FCollection::IsUsingSourceControl() const
+{
+	return bUseSCC;
+}
diff --git a/Engine/Source/Developer/CollectionManager/Private/CollectionManager.cpp b/Engine/Source/Developer/CollectionManager/Private/CollectionManager.cpp
--- a/Engine/Source/Developer/CollectionManager/Private/CollectionManager.cpp
+++ b/Engine/Source/Developer/CollectionManager/Private/CollectionManager.cpp
@@ -1,10 +1,12 @@
 // Copyright Epic Games, Inc. All Rights Reserved.
 
 #include "CollectionManager.h"
+#include "CollectionManagerTelemetry.h"
 #include "HAL/FileManager.h"
 #include "Misc/Paths.h"
 #include "Containers/Ticker.h"
 #include "CollectionManagerLog.h"
+#include "CollectionManagerModule.h"
 #include "FileCache.h"
 #include "Misc/FileHelper.h"
 #include "Misc/ScopeRWLock.h"
@@ -1029,6 +1031,9 @@ bool FCollectionManager::AddToCollection(FName CollectionName, ECollectionShareT
 		return false;
 	}
 
+	using namespace UE;
+	auto TelemetryEvent = UE::Analytics::MakeScopedTimedTelemetryEvent<UE::CollectionManager::Analytics::FAddObjects>(FCollectionManagerModule::GetModule().OnAnalyticsEventEmitted);
+
 	int32 NumAdded = 0;
 	for (const FSoftObjectPath& ObjectPath : ObjectPaths)
 	{
@@ -1038,6 +1043,11 @@ bool FCollectionManager::AddToCollection(FName CollectionName, ECollectionShareT
 		}
 	}
 
+	if ((*CollectionRefPtr)->IsUsingSourceControl())
+	{
+		TelemetryEvent.SetStatusCode(TelemetryEvent.GetStatusCode() | CollectionManager::Analytics::StatusCodeFields::UsesSCC);
+	}
+
 	if (NumAdded > 0)
 	{
 		if (InternalSaveCollection(*CollectionRefPtr, LastError))
@@ -1058,6 +1068,7 @@ PRAGMA_DISABLE_DEPRECATION_WARNINGS
 				AssetsAddedEvent.Broadcast(CollectionKey, UE::SoftObjectPath::Private::ConvertSoftObjectPaths(ObjectPaths));
 			}
 PRAGMA_ENABLE_DEPRECATION_WARNINGS
+			TelemetryEvent.SetStatusCode(TelemetryEvent.GetStatusCode() | CollectionManager::Analytics::StatusCodeFields::Success);
 			return true;
 		}
 		else
@@ -1114,6 +1125,9 @@ bool FCollectionManager::RemoveFromCollection(FName CollectionName, ECollectionS
 		return false;
 	}
 
+	using namespace UE;
+	auto TelemetryEvent = Analytics::MakeScopedTimedTelemetryEvent<CollectionManager::Analytics::FRemoveObjects>(FCollectionManagerModule::GetModule().OnAnalyticsEventEmitted);
+
 	TArray<FSoftObjectPath> RemovedAssets;
 	for (const FSoftObjectPath& ObjectPath : ObjectPaths)
 	{
@@ -1123,6 +1137,11 @@ bool FCollectionManager::RemoveFromCollection(FName CollectionName, ECollectionS
 		}
 	}
 
+	if ((*CollectionRefPtr)->IsUsingSourceControl())
+	{
+		TelemetryEvent.SetStatusCode(TelemetryEvent.GetStatusCode() | CollectionManager::Analytics::StatusCodeFields::UsesSCC);
+	}
+
 	if (RemovedAssets.Num() == 0)
 	{
 		// Failed to remove, none of the objects were in the collection
@@ -1148,6 +1167,10 @@ PRAGMA_DISABLE_DEPRECATION_WARNINGS
 			AssetsRemovedEvent.Broadcast(CollectionKey, UE::SoftObjectPath::Private::ConvertSoftObjectPaths(ObjectPaths));
 		}
 PRAGMA_ENABLE_DEPRECATION_WARNINGS
+		if ((*CollectionRefPtr)->IsUsingSourceControl())
+		{
+			TelemetryEvent.SetStatusCode(TelemetryEvent.GetStatusCode() | CollectionManager::Analytics::StatusCodeFields::Success);
+		}
 		return true;
 	}
 	else
diff --git a/Engine/Source/Developer/CollectionManager/Public/CollectionManagerModule.h b/Engine/Source/Developer/CollectionManager/Public/CollectionManagerModule.h
--- a/Engine/Source/Developer/CollectionManager/Public/CollectionManagerModule.h
+++ b/Engine/Source/Developer/CollectionManager/Public/CollectionManagerModule.h
@@ -3,6 +3,7 @@
 #pragma once
 
 #include "CoreMinimal.h"
+#include "CollectionManagerTelemetry.h"
 #include "Modules/ModuleInterface.h"
 #include "Modules/ModuleManager.h"
 #include "UObject/NameTypes.h"
@@ -48,6 +49,8 @@ public:
 		return FModuleManager::Get().IsModuleLoaded(CollectionManagerModuleName);
 	}
 
+	UE::CollectionManager::Analytics::FEmitAnalyticsEvent OnAnalyticsEventEmitted;
+
 private:
 	ICollectionManager* CollectionManager;
 	class FCollectionManagerConsoleCommands* ConsoleCommands;
diff --git a/Engine/Source/Developer/CollectionManager/Public/CollectionManagerTelemetry.h b/Engine/Source/Developer/CollectionManager/Public/CollectionManagerTelemetry.h
new file mode 100644
--- /dev/null
+++ b/Engine/Source/Developer/CollectionManager/Public/CollectionManagerTelemetry.h
@@ -0,0 +1,31 @@
+// Copyright Epic Games, Inc. All Rights Reserved.
+
+#pragma once
+
+#include <variant>
+
+#include "AnalyticsTimedTelemetryEvents.h"
+#include "Delegates/Delegate.h"
+
+/**
+ * Events that the CollectionManager emits for telemetry tracking purposes
+ */
+namespace UE::CollectionManager::Analytics
+{
+	/**
+	 * StatusCode bitfield positions
+	 */
+	struct StatusCodeFields
+	{
+		// Operation uses Source control
+		static constexpr uint64 UsesSCC = 1llu << 62;
+		// Operation was successful
+		static constexpr uint64 Success = 1llu << 63;
+	};
+	struct FAddObjects : UE::Analytics::FTimedTelemetryEvent {};
+	struct FRemoveObjects : UE::Analytics::FTimedTelemetryEvent {};
+	
+	using FEvent = std::variant<FAddObjects, FRemoveObjects>;
+
+	DECLARE_DELEGATE_OneParam(FEmitAnalyticsEvent, const FEvent&)
+}
\ No newline at end of file
diff --git a/Engine/Source/Runtime/Analytics/Analytics/Public/AnalyticsTimedTelemetryEvents.h b/Engine/Source/Runtime/Analytics/Analytics/Public/AnalyticsTimedTelemetryEvents.h
new file mode 100644
--- /dev/null
+++ b/Engine/Source/Runtime/Analytics/Analytics/Public/AnalyticsTimedTelemetryEvents.h
@@ -0,0 +1,184 @@
+﻿#pragma once
+
+#include "HAL/PlatformTime.h"
+
+namespace UE::Analytics
+{
+	// Generic telemetry data for a timed telemetry event
+	struct FTimedTelemetryEvent
+	{
+		// Status code that can be used by an event to convey some custom data
+		uint64 StatusCode = 0;
+		// Duration of the event in seconds
+		double Duration = 0.0;
+	};
+	
+	/**
+	 * A tracker of timespans with finite duration to send to telemetry
+	 * It is intended for relatively long-duration timespans to be measured in seconds to minutes using
+	 * a module which registers an IAnalyticsProviderET and starts an analytics session
+	 * For profiling at a fine-grained level, use the Trace system
+	 */
+	template<typename EventT, typename EventDelegateT>
+	class TTimedTelemetryEvent
+	{
+	public:
+		TTimedTelemetryEvent() {}
+		/**
+		 * EventDelegate is the callback to be used when the timespan completes
+		 */
+		explicit TTimedTelemetryEvent(const EventDelegateT& EventDelegate);
+		TTimedTelemetryEvent(TTimedTelemetryEvent&& Other);
+		TTimedTelemetryEvent& operator=(TTimedTelemetryEvent&& Other);
+	
+		TTimedTelemetryEvent(TTimedTelemetryEvent&) = delete;
+		TTimedTelemetryEvent& operator=(TTimedTelemetryEvent&) = delete;
+
+		/**
+		 * Marks an explicit beginning of the timespan to measure
+		 */
+		void MarkBegin();
+		/**
+		 * Marks an explicit ending of the timespan to measure
+		 */
+		void MarkEnd();
+
+		/**
+		 * Sets the StatusCode of the event
+		 * The StatusCode can be used to convey limited information about the event itself
+		 */
+		void SetStatusCode(uint64 Code) { StatusCode = Code; }
+
+		/**
+		 * Get the StatusCode of the event
+		 * The StatusCode can be used to convey limited information about the event itself
+		 */
+		uint64 GetStatusCode() const { return StatusCode; }
+
+		/**
+		 * Mutable access to the underlying event to set additional metadata
+		 */
+		EventT& GetEvent() { return Event; }
+	
+	private:
+		//~ Track begin and end to prevent multiple events from being sent
+		bool MarkedBegin = false;
+		bool MarkedEnd = false;
+		// Cache StatusCode until the event ends
+		uint64 StatusCode = 0;
+		// Cache BeginTime of event to compute duration
+		double BeginTime = 0.0;
+		EventT Event;
+		const EventDelegateT* EventDelegate = nullptr;
+	};
+
+
+	template <typename EventT, typename EventDelegateT>
+	TTimedTelemetryEvent<EventT, EventDelegateT>::TTimedTelemetryEvent(const EventDelegateT& InEventDelegate)
+		: EventDelegate(&InEventDelegate)
+	{
+		static_assert(TIsDerivedFrom<EventT, FTimedTelemetryEvent>::Value, "EventT must be derived from FTimedTelemetryEvent");
+	}
+
+	template <typename EventT, typename EventDelegateT>
+	TTimedTelemetryEvent<EventT, EventDelegateT>::TTimedTelemetryEvent(TTimedTelemetryEvent&& Other)
+	{
+		MarkedBegin = Other.MarkedBegin;
+		MarkedEnd = Other.MarkedEnd;
+		StatusCode = Other.StatusCode;
+		BeginTime = Other.BeginTime;
+		Event = MoveTemp(Other.Event);
+		EventDelegate = Other.EventDelegate;
+
+		Other.MarkedBegin = false;
+		Other.MarkedEnd = false;
+		Other.StatusCode = 0;
+		Other.BeginTime = 0.0;
+		Other.EventDelegate = nullptr;
+	}
+	
+	template <typename EventT, typename EventDelegateT>
+	TTimedTelemetryEvent<EventT, EventDelegateT>& TTimedTelemetryEvent<EventT, EventDelegateT>::operator=(TTimedTelemetryEvent&& Other)
+	{
+		if (this != &Other)
+		{
+			MarkedBegin = Other.MarkedBegin;
+			MarkedEnd = Other.MarkedEnd;
+			StatusCode = Other.StatusCode;
+			BeginTime = Other.BeginTime;
+			Event = MoveTemp(Other.Event);
+			EventDelegate = Other.EventDelegate;
+
+			Other.MarkedBegin = false;
+			Other.MarkedEnd = false;
+			Other.StatusCode = 0;
+			Other.BeginTime = 0.0;
+			Other.EventDelegate = nullptr;
+		}
+		return *this;
+	}
+	
+	template <typename EventT, typename EventDelegateT>
+	void TTimedTelemetryEvent<EventT, EventDelegateT>::MarkBegin()
+	{
+		if (EventDelegate->IsBound() && !MarkedEnd && !MarkedBegin)
+		{
+			MarkedBegin = true;
+			BeginTime = FPlatformTime::Seconds();
+		}
+	}
+
+	template <typename EventT, typename EventDelegateT>
+	void TTimedTelemetryEvent<EventT, EventDelegateT>::MarkEnd()
+	{
+		if (MarkedBegin && !MarkedEnd)
+		{
+			// NOTE: May result in large durations on processes that suspend
+			Event.Duration = FPlatformTime::Seconds() - BeginTime;
+			// Send the cached ErrorCode along with the event
+			Event.StatusCode = StatusCode;
+
+			// No need to check IsBound - shouldn't be possible to be here if not bound
+			EventDelegate->Execute(Event);
+		}
+	}
+
+	// A RAII Wrapper for TimedTelemetryEvents
+	template<typename EventT, typename EventDelegateT>
+	class TScopedTimedTelemetryEvent
+	{
+	private:
+		using TimedTelemetryEventT = TTimedTelemetryEvent<EventT, EventDelegateT>;
+	public:		
+		TScopedTimedTelemetryEvent(const EventDelegateT& EventDelegate);
+		~TScopedTimedTelemetryEvent();
+
+		void SetStatusCode(uint64 StatusCode) { TimedTelemetryEvent.SetStatusCode(StatusCode); }
+		uint64 GetStatusCode() const { return TimedTelemetryEvent.GetStatusCode(); }
+
+		EventT& GetEvent() { return TimedTelemetryEvent.GetEvent(); }
+
+	private:
+		TimedTelemetryEventT TimedTelemetryEvent;
+	};
+
+	template<typename EventT, typename EventDelegateT>
+	inline TScopedTimedTelemetryEvent<EventT, EventDelegateT>::TScopedTimedTelemetryEvent(const EventDelegateT& EventDelegate)
+		: TimedTelemetryEvent(EventDelegate)
+	{
+		TimedTelemetryEvent.MarkBegin();
+	}
+
+	template<typename EventT, typename EventDelegateT>
+	inline TScopedTimedTelemetryEvent<EventT, EventDelegateT>::~TScopedTimedTelemetryEvent()
+	{
+		TimedTelemetryEvent.MarkEnd();
+	}
+
+	// Helper to create a Scoped Timed Telemetry Event
+	template<typename EventT, typename EventDelegateT>
+	TScopedTimedTelemetryEvent<EventT, EventDelegateT> MakeScopedTimedTelemetryEvent(const EventDelegateT& EventDelegate)
+	{
+		return TScopedTimedTelemetryEvent<EventT, EventDelegateT>(EventDelegate);
+	}
+}
\ No newline at end of file
