Add copyright notice to AnalyticsTimedTelemetryEvents.h

diff --git a/Engine/Source/Runtime/Analytics/Analytics/Public/AnalyticsTimedTelemetryEvents.h b/Engine/Source/Runtime/Analytics/Analytics/Public/AnalyticsTimedTelemetryEvents.h
--- a/Engine/Source/Runtime/Analytics/Analytics/Public/AnalyticsTimedTelemetryEvents.h
+++ b/Engine/Source/Runtime/Analytics/Analytics/Public/AnalyticsTimedTelemetryEvents.h
@@ -1,4 +1,6 @@
-﻿#pragma once
+﻿// Copyright Epic Games, Inc. All Rights Reserved.
+
+#pragma once
 
 #include "HAL/PlatformTime.h"
 
