Added support for Unreal Insights Regions in Analytics Flow Tracker

diff --git a/Engine/Source/Runtime/Analytics/AnalyticsET/Private/AnalyticsFlowTracker.cpp b/Engine/Source/Runtime/Analytics/AnalyticsET/Private/AnalyticsFlowTracker.cpp
--- a/Engine/Source/Runtime/Analytics/AnalyticsET/Private/AnalyticsFlowTracker.cpp
+++ b/Engine/Source/Runtime/Analytics/AnalyticsET/Private/AnalyticsFlowTracker.cpp
@@ -19,7 +19,8 @@ void FAnalyticsFlowTracker::EndSession()
 	// End all the open flows from the stack
 	while (FlowDataStack.Num())
 	{
-		EndFlowInternal(FlowDataStack.Last());
+		// We are forcibly shutting down the sub flows so mark them as unsuccesful
+		EndFlowInternal(FlowDataStack.Last(), false);
 	}
 
 	ensure(FlowDataRegistry.IsEmpty());
@@ -31,7 +32,7 @@ void FAnalyticsFlowTracker::EndSession()
 FGuid FAnalyticsFlowTracker::StartFlow(const FName& NewFlowName)
 {
 	FScopeLock ScopeLock(&CriticalSection);
-	TRACE_BOOKMARK(TEXT("STARTFLOW: %s"), *NewFlowName.ToString());
+	TRACE_BEGIN_REGION(*FString::Printf(TEXT("Flow: %s"), *NewFlowName.ToString()));
 
 	// Create a new Guid for this flow, can we assume it is unique?
 	FGuid NewFlowGuid = FGuid::NewGuid();
@@ -76,7 +77,7 @@ FGuid FAnalyticsFlowTracker::StartSubFlowInternal(const FName& NewSubFlowName, c
 
 	if (ensureMsgf(FlowData, TEXT("SubFlow started outside of a valid flow scope")))
 	{
-		TRACE_BOOKMARK(TEXT("STARTSUBFLOW: %s"), *NewSubFlowName.ToString());
+		TRACE_END_REGION(*FString::Printf(TEXT("Flow: %s"), *FlowData->FlowName.ToString()));
 
 		// Create a new Guid for this SubFlow, can we assume it is unique?
 		FGuid NewSubFlowGuid = FGuid::NewGuid();
@@ -149,7 +150,7 @@ bool FAnalyticsFlowTracker::EndSubFlowInternal(const FGuid& SubFlowGuid, bool bS
 			AnalyticsProvider->RecordEvent(SubFlowEventName, EventAttributes);
 		}
 
-		TRACE_BOOKMARK(TEXT("ENDSUBFLOW: %s"), *SubFlowData->SubFlowName.ToString());
+		TRACE_END_REGION(*FString::Printf(TEXT("Flow: %s"), *SubFlowData->SubFlowName.ToString()));
 
 		FFlowData* FlowData = FlowDataRegistry.Find(FlowGuid);
 
@@ -299,7 +300,7 @@ bool FAnalyticsFlowTracker::EndFlowInternal(const FGuid& FlowGuid, bool bSuccess
 		AnalyticsProvider->FlushEvents();
 	}
 
-	TRACE_BOOKMARK(TEXT("ENDFLOW: %s"), *FlowData->FlowName.ToString());
+	TRACE_END_REGION(*FString::Printf(TEXT("Flow: %s"), *FlowData->FlowName.ToString()));
 
 	// Clear up our data
 	for (FGuid SubFlowGuid : FlowData->SubFlowDataArray)
