[Engine]

diff --git a/Engine/Source/Runtime/AIModule/Classes/Perception/AISightTargetInterface.h b/Engine/Source/Runtime/AIModule/Classes/Perception/AISightTargetInterface.h
--- a/Engine/Source/Runtime/AIModule/Classes/Perception/AISightTargetInterface.h
+++ b/Engine/Source/Runtime/AIModule/Classes/Perception/AISightTargetInterface.h
@@ -10,13 +10,13 @@
 
 class AActor;
 
-UINTERFACE()
-class AIMODULE_API UAISightTargetInterface : public UInterface
+UINTERFACE(MinimalAPI)
+class UAISightTargetInterface : public UInterface
 {
 	GENERATED_UINTERFACE_BODY()
 };
 
-struct AIMODULE_API FCanBeSeenFromContext
+struct FCanBeSeenFromContext
 {
 	/** The query identifier used by the delegate call to find the appropriate query */
 	FAISightQueryID SightQueryID;
@@ -28,7 +28,7 @@ struct AIMODULE_API FCanBeSeenFromContext
 	const bool* bWasVisible = nullptr;
 };
 
-class AIMODULE_API IAISightTargetInterface
+class IAISightTargetInterface
 {
 	GENERATED_IINTERFACE_BODY()
 
diff --git a/Engine/Source/Runtime/AIModule/Classes/Perception/PawnSensingComponent.h b/Engine/Source/Runtime/AIModule/Classes/Perception/PawnSensingComponent.h
--- a/Engine/Source/Runtime/AIModule/Classes/Perception/PawnSensingComponent.h
+++ b/Engine/Source/Runtime/AIModule/Classes/Perception/PawnSensingComponent.h
@@ -17,8 +17,8 @@ class UPawnNoiseEmitterComponent;
  * SensingComponent encapsulates sensory (ie sight and hearing) settings and functionality for an Actor,
  * allowing the actor to see/hear Pawns in the world. It does nothing on network clients.
  */
-UCLASS(ClassGroup=AI, HideCategories=(Activation, "Components|Activation", Collision), meta=(BlueprintSpawnableComponent))
-class AIMODULE_API UPawnSensingComponent : public UActorComponent
+UCLASS(ClassGroup=AI, HideCategories=(Activation, "Components|Activation", Collision), meta=(BlueprintSpawnableComponent), MinimalAPI)
+class UPawnSensingComponent : public UActorComponent
 {
 	GENERATED_UCLASS_BODY()
 
@@ -53,21 +53,21 @@ public:
 	 * A value <= 0 prevents any updates.
 	 */
 	UFUNCTION(BlueprintCallable, BlueprintAuthorityOnly, Category="AI|Components|PawnSensing")
-	virtual void SetSensingInterval(const float NewSensingInterval);
+	AIMODULE_API virtual void SetSensingInterval(const float NewSensingInterval);
 
 	/** Enables or disables sensing updates. The timer is reset in either case. */
 	UFUNCTION(BlueprintCallable, BlueprintAuthorityOnly, Category="AI|Components|PawnSensing")
-	virtual void SetSensingUpdatesEnabled(const bool bEnabled);
+	AIMODULE_API virtual void SetSensingUpdatesEnabled(const bool bEnabled);
 
 	/** Sets PeripheralVisionAngle. Calculates PeripheralVisionCosine from PeripheralVisionAngle */
 	UFUNCTION(BlueprintCallable, BlueprintAuthorityOnly, Category="AI|Components|PawnSensing")
-	virtual void SetPeripheralVisionAngle(const float NewPeripheralVisionAngle);
+	AIMODULE_API virtual void SetPeripheralVisionAngle(const float NewPeripheralVisionAngle);
 
 	UFUNCTION(BlueprintCallable, Category="AI|Components|PawnSensing")
-	float GetPeripheralVisionAngle() const;
+	AIMODULE_API float GetPeripheralVisionAngle() const;
 
 	UFUNCTION(BlueprintCallable, Category="AI|Components|PawnSensing")
-	float GetPeripheralVisionCosine() const;
+	AIMODULE_API float GetPeripheralVisionCosine() const;
 
 	/** If true, component will perform sensing updates. At runtime change this using SetSensingUpdatesEnabled(). */
 	UPROPERTY(EditAnywhere, BlueprintReadOnly, Category=AI)
@@ -89,23 +89,23 @@ public:
 	uint32 bHearNoises:1;
 
 	/** Is the given actor our owner? Used to ensure that we are not trying to sense our self / our owner. */
-	virtual bool IsSensorActor(const AActor* Actor) const;
+	AIMODULE_API virtual bool IsSensorActor(const AActor* Actor) const;
 
 	/** Are we capable of sensing anything (and do we have any callbacks that care about sensing)? If so UpdateAISensing() will be called every sensing interval. */
-	virtual bool CanSenseAnything() const;
+	AIMODULE_API virtual bool CanSenseAnything() const;
 
 	/** Returns true if we should check whether the given Pawn is visible (because we can see things, the Pawn is not hidden, and if the Pawn is a player and we only see players) */
-	virtual bool ShouldCheckVisibilityOf(APawn* Pawn) const;
+	AIMODULE_API virtual bool ShouldCheckVisibilityOf(APawn* Pawn) const;
 
 	/** 
 	 * Chance of seeing other pawn decreases with increasing distance or angle in peripheral vision
 	 * @param bMaySkipChecks if true allows checks to be sometimes skipped if target is far away (used by periodic automatic visibility checks)
 	 * @return true if the specified pawn Other is potentially visible (within peripheral vision, etc.) - still need to do LineOfSightTo() check to see if actually visible.  
 	 */
-	virtual bool CouldSeePawn(const APawn* Other, bool bMaySkipChecks = false) const;
+	AIMODULE_API virtual bool CouldSeePawn(const APawn* Other, bool bMaySkipChecks = false) const;
 
 	/** Returns true if we should check whether we can hear the given Pawn (because we are able to hear, and the Pawn has the correct team relationship to us) */
-	virtual bool ShouldCheckAudibilityOf(APawn* Pawn) const;
+	AIMODULE_API virtual bool ShouldCheckAudibilityOf(APawn* Pawn) const;
 
 	/**
 	 * Check line to other actor.
@@ -113,46 +113,46 @@ public:
 	 * @param ViewPoint is eye position visibility is being checked from.
 	 * @return true if controller's pawn can see Other actor.
 	 */
-	virtual bool HasLineOfSightTo(const AActor* Other) const;
+	AIMODULE_API virtual bool HasLineOfSightTo(const AActor* Other) const;
 
 	/** Test whether the noise is loud enough and recent enough to care about.  bSourceWithinNoiseEmitter is true iff the 
 	 * noise was made by the pawn itself or within close proximity (its collision volume).  Otherwise the noise was made
 	 * at significant distance from the pawn.
 	 */
-	bool IsNoiseRelevant(const APawn& Pawn, const UPawnNoiseEmitterComponent& NoiseEmitterComponent, bool bSourceWithinNoiseEmitter) const;
+	AIMODULE_API bool IsNoiseRelevant(const APawn& Pawn, const UPawnNoiseEmitterComponent& NoiseEmitterComponent, bool bSourceWithinNoiseEmitter) const;
 
 	/** @Returns true if sensor can hear this noise. Only executed if the noise has been determined to be relevant (via IsNoiseRelevant) */
-	virtual bool CanHear(const FVector& NoiseLoc, float Loudness, bool bFailedLOS) const;
+	AIMODULE_API virtual bool CanHear(const FVector& NoiseLoc, float Loudness, bool bFailedLOS) const;
 
 	//~ Begin UActorComponent Interface.
-	virtual void InitializeComponent() override;
+	AIMODULE_API virtual void InitializeComponent() override;
 	//~ End UActorComponent Interface.
 
 	/** Get position where hearing/seeing occurs (i.e. ear/eye position).  If we ever need different positions for hearing/seeing, we'll deal with that then! */
-	virtual FVector GetSensorLocation() const;
+	AIMODULE_API virtual FVector GetSensorLocation() const;
 
 	/**  Get the rotation of this sensor. We need this for the sight component */
-	virtual FRotator GetSensorRotation() const;
+	AIMODULE_API virtual FRotator GetSensorRotation() const;
 
 protected:
 
 	/** See if there are interesting sounds and sights that we want to detect, and respond to them if so. */
-	virtual void SensePawn(APawn& Pawn);
+	AIMODULE_API virtual void SensePawn(APawn& Pawn);
 
 	/** Update function called on timer intervals. */
-	virtual void OnTimer();
+	AIMODULE_API virtual void OnTimer();
 
 	/** Handle for efficient management of OnTimer timer */
 	FTimerHandle TimerHandle_OnTimer;
 
 	/** Modify the timer to fire in TimeDelay seconds. A value <= 0 disables the timer. */
-	virtual void SetTimer(const float TimeDelay);
+	AIMODULE_API virtual void SetTimer(const float TimeDelay);
 
 	/** Calls SensePawn on any Pawns that we are allowed to sense. */
-	virtual void UpdateAISensing();
+	AIMODULE_API virtual void UpdateAISensing();
 
-	AActor* GetSensorActor() const;	// Get the actor used as the actual sensor location is derived from this actor.
-	AController* GetSensorController() const; // Get the controller of the sensor actor.
+	AIMODULE_API AActor* GetSensorActor() const;	// Get the actor used as the actual sensor location is derived from this actor.
+	AIMODULE_API AController* GetSensorController() const; // Get the controller of the sensor actor.
 
 public:
 	/** Delegate to execute when we see a Pawn. */
@@ -166,13 +166,13 @@ public:
 protected:
 
 	// Broadcasts notification that our sensor sees a Pawn, using the OnSeePawn delegates.
-	virtual void BroadcastOnSeePawn(APawn& Pawn);
+	AIMODULE_API virtual void BroadcastOnSeePawn(APawn& Pawn);
 
 	// Broadcasts notification that our sensor hears a noise made local to a Pawn's position, using the OnHearNoise delegates.
-	virtual void BroadcastOnHearLocalNoise(APawn& Instigator, const FVector& Location, float Volume);
+	AIMODULE_API virtual void BroadcastOnHearLocalNoise(APawn& Instigator, const FVector& Location, float Volume);
 
 	// Broadcasts notification that our sensor hears a noise made remotely from a Pawn's position, using the OnHearNoise delegates.
-	virtual void BroadcastOnHearRemoteNoise(APawn& Instigator, const FVector& Location, float Volume);
+	AIMODULE_API virtual void BroadcastOnHearRemoteNoise(APawn& Instigator, const FVector& Location, float Volume);
 
 	/** How far to the side AI can see, in degrees. Use SetPeripheralVisionAngle to change the value at runtime. */
 	UPROPERTY(EditAnywhere, BlueprintReadOnly, Category=AI)
diff --git a/Engine/Source/Runtime/AIModule/Classes/Tasks/AITask.h b/Engine/Source/Runtime/AIModule/Classes/Tasks/AITask.h
--- a/Engine/Source/Runtime/AIModule/Classes/Tasks/AITask.h
+++ b/Engine/Source/Runtime/AIModule/Classes/Tasks/AITask.h
@@ -20,8 +20,8 @@ enum class EAITaskPriority : uint8
 	Ultimate = 254,
 };
 
-UCLASS(Abstract, BlueprintType)
-class AIMODULE_API UAITask : public UGameplayTask
+UCLASS(Abstract, BlueprintType, MinimalAPI)
+class UAITask : public UGameplayTask
 {
 	GENERATED_BODY()
 protected:
@@ -29,19 +29,19 @@ protected:
 	UPROPERTY(BlueprintReadOnly, Category="AI|Tasks")
 	TObjectPtr<AAIController> OwnerController;
 
-	virtual void Activate() override;
+	AIMODULE_API virtual void Activate() override;
 
 public:
-	UAITask(const FObjectInitializer& ObjectInitializer);
+	AIMODULE_API UAITask(const FObjectInitializer& ObjectInitializer);
 
-	static AAIController* GetAIControllerForActor(AActor* Actor);
+	static AIMODULE_API AAIController* GetAIControllerForActor(AActor* Actor);
 	AAIController* GetAIController() const { return OwnerController; };
 
-	void InitAITask(AAIController& AIOwner, IGameplayTaskOwnerInterface& InTaskOwner, uint8 InPriority);
-	void InitAITask(AAIController& AIOwner, IGameplayTaskOwnerInterface& InTaskOwner);
+	AIMODULE_API void InitAITask(AAIController& AIOwner, IGameplayTaskOwnerInterface& InTaskOwner, uint8 InPriority);
+	AIMODULE_API void InitAITask(AAIController& AIOwner, IGameplayTaskOwnerInterface& InTaskOwner);
 
 	/** effectively adds UAIResource_Logic to the set of Claimed resources */
-	void RequestAILogicLocking();
+	AIMODULE_API void RequestAILogicLocking();
 
 	template <class T>
 	static T* NewAITask(AAIController& AIOwner, IGameplayTaskOwnerInterface& InTaskOwner, FName InstanceName = FName())
diff --git a/Engine/Source/Runtime/AIModule/Classes/Tasks/AITask_LockLogic.h b/Engine/Source/Runtime/AIModule/Classes/Tasks/AITask_LockLogic.h
--- a/Engine/Source/Runtime/AIModule/Classes/Tasks/AITask_LockLogic.h
+++ b/Engine/Source/Runtime/AIModule/Classes/Tasks/AITask_LockLogic.h
@@ -7,8 +7,8 @@
 #include "AITask_LockLogic.generated.h"
 
 /** Locks AI logic until removed by external trigger */
-UCLASS()
-class AIMODULE_API UAITask_LockLogic : public UAITask
+UCLASS(MinimalAPI)
+class UAITask_LockLogic : public UAITask
 {
 	GENERATED_UCLASS_BODY()
 };
diff --git a/Engine/Source/Runtime/AIModule/Classes/Tasks/AITask_MoveTo.h b/Engine/Source/Runtime/AIModule/Classes/Tasks/AITask_MoveTo.h
--- a/Engine/Source/Runtime/AIModule/Classes/Tasks/AITask_MoveTo.h
+++ b/Engine/Source/Runtime/AIModule/Classes/Tasks/AITask_MoveTo.h
@@ -13,26 +13,26 @@ class AAIController;
 
 DECLARE_DYNAMIC_MULTICAST_DELEGATE_TwoParams(FMoveTaskCompletedSignature, TEnumAsByte<EPathFollowingResult::Type>, Result, AAIController*, AIController);
 
-UCLASS()
-class AIMODULE_API UAITask_MoveTo : public UAITask
+UCLASS(MinimalAPI)
+class UAITask_MoveTo : public UAITask
 {
 	GENERATED_BODY()
 
 public:
-	UAITask_MoveTo(const FObjectInitializer& ObjectInitializer);
+	AIMODULE_API UAITask_MoveTo(const FObjectInitializer& ObjectInitializer);
 
 	/** tries to start move request and handles retry timer */
-	void ConditionalPerformMove();
+	AIMODULE_API void ConditionalPerformMove();
 
 	/** prepare move task for activation */
-	void SetUp(AAIController* Controller, const FAIMoveRequest& InMoveRequest);
+	AIMODULE_API void SetUp(AAIController* Controller, const FAIMoveRequest& InMoveRequest);
 
 	EPathFollowingResult::Type GetMoveResult() const { return MoveResult; }
 	bool WasMoveSuccessful() const { return MoveResult == EPathFollowingResult::Success; }
 	bool WasMovePartial() const { return Path.IsValid() && Path->IsPartial(); }
 
 	UFUNCTION(BlueprintCallable, Category = "AI|Tasks", meta = (AdvancedDisplay = "AcceptanceRadius,StopOnOverlap,AcceptPartialPath,bUsePathfinding,bUseContinuousGoalTracking,ProjectGoalOnNavigation", DefaultToSelf = "Controller", BlueprintInternalUseOnly = "TRUE", DisplayName = "Move To Location or Actor"))
-	static UAITask_MoveTo* AIMoveTo(AAIController* Controller, FVector GoalLocation, AActor* GoalActor = nullptr,
+	static AIMODULE_API UAITask_MoveTo* AIMoveTo(AAIController* Controller, FVector GoalLocation, AActor* GoalActor = nullptr,
 		float AcceptanceRadius = -1.f, EAIOptionFlag::Type StopOnOverlap = EAIOptionFlag::Default, EAIOptionFlag::Type AcceptPartialPath = EAIOptionFlag::Default,
 		bool bUsePathfinding = true, bool bLockAILogic = true, bool bUseContinuousGoalTracking = false, EAIOptionFlag::Type ProjectGoalOnNavigation = EAIOptionFlag::Default);
 
@@ -41,7 +41,7 @@ public:
 	FAIMoveRequest& GetMoveRequestRef() { return MoveRequest; }
 
 	/** Switch task into continuous tracking mode: keep restarting move toward goal actor. Only pathfinding failure or external cancel will be able to stop this task. */
-	void SetContinuousGoalTracking(bool bEnable);
+	AIMODULE_API void SetContinuousGoalTracking(bool bEnable);
 
 protected:
 	UPROPERTY(BlueprintAssignable)
@@ -75,33 +75,33 @@ protected:
 	TEnumAsByte<EPathFollowingResult::Type> MoveResult;
 	uint8 bUseContinuousTracking : 1;
 
-	virtual void Activate() override;
-	virtual void OnDestroy(bool bOwnerFinished) override;
+	AIMODULE_API virtual void Activate() override;
+	AIMODULE_API virtual void OnDestroy(bool bOwnerFinished) override;
 
-	virtual void Pause() override;
-	virtual void Resume() override;
+	AIMODULE_API virtual void Pause() override;
+	AIMODULE_API virtual void Resume() override;
 
 	/** finish task */
-	void FinishMoveTask(EPathFollowingResult::Type InResult);
+	AIMODULE_API void FinishMoveTask(EPathFollowingResult::Type InResult);
 
 	/** stores path and starts observing its events */
-	void SetObservedPath(FNavPathSharedPtr InPath);
+	AIMODULE_API void SetObservedPath(FNavPathSharedPtr InPath);
 
 	/** remove all delegates */
-	virtual void ResetObservers();
+	AIMODULE_API virtual void ResetObservers();
 
 	/** remove all timers */
-	virtual void ResetTimers();
+	AIMODULE_API virtual void ResetTimers();
 
 	/** tries to update invalidated path and handles retry timer */
-	void ConditionalUpdatePath();
+	AIMODULE_API void ConditionalUpdatePath();
 
 	/** start move request */
-	virtual void PerformMove();
+	AIMODULE_API virtual void PerformMove();
 
 	/** event from followed path */
-	virtual void OnPathEvent(FNavigationPath* InPath, ENavPathEvent::Type Event);
+	AIMODULE_API virtual void OnPathEvent(FNavigationPath* InPath, ENavPathEvent::Type Event);
 
 	/** event from path following */
-	virtual void OnRequestFinished(FAIRequestID RequestID, const FPathFollowingResult& Result);
+	AIMODULE_API virtual void OnRequestFinished(FAIRequestID RequestID, const FPathFollowingResult& Result);
 };
diff --git a/Engine/Source/Runtime/AIModule/Classes/Tasks/AITask_RunEQS.h b/Engine/Source/Runtime/AIModule/Classes/Tasks/AITask_RunEQS.h
--- a/Engine/Source/Runtime/AIModule/Classes/Tasks/AITask_RunEQS.h
+++ b/Engine/Source/Runtime/AIModule/Classes/Tasks/AITask_RunEQS.h
@@ -11,24 +11,24 @@
 class AAIController;
 class UEnvQuery;
 
-UCLASS()
-class AIMODULE_API UAITask_RunEQS : public UAITask
+UCLASS(MinimalAPI)
+class UAITask_RunEQS : public UAITask
 {
 	GENERATED_BODY()
 		
 public:
-	UAITask_RunEQS(const FObjectInitializer& ObjectInitializer = FObjectInitializer::Get());
+	AIMODULE_API UAITask_RunEQS(const FObjectInitializer& ObjectInitializer = FObjectInitializer::Get());
 
 	UFUNCTION(BlueprintCallable, Category = "AI|Tasks", meta = (DefaultToSelf = "Controller", BlueprintInternalUseOnly = "TRUE"))
-	static UAITask_RunEQS* RunEQS(AAIController* Controller, UEnvQuery* QueryTemplate);
+	static AIMODULE_API UAITask_RunEQS* RunEQS(AAIController* Controller, UEnvQuery* QueryTemplate);
 
 	void SetQueryTemplate(UEnvQuery& InQueryTemplate) { EQSRequest.QueryTemplate = &InQueryTemplate; }
 	void SetNotificationDelegate(const FQueryFinishedSignature& InNotificationDelegate) { NotificationDelegate = InNotificationDelegate; }
 		
 protected:
-	virtual void Activate() override;
+	AIMODULE_API virtual void Activate() override;
 
-	void OnEQSRequestFinished(TSharedPtr<FEnvQueryResult> Result);
+	AIMODULE_API void OnEQSRequestFinished(TSharedPtr<FEnvQueryResult> Result);
 
 	FEQSParametrizedQueryExecutionRequest EQSRequest;
 	FQueryFinishedSignature EQSFinishedDelegate;
diff --git a/Engine/Source/Runtime/AIModule/Classes/VisualLoggerExtension.h b/Engine/Source/Runtime/AIModule/Classes/VisualLoggerExtension.h
--- a/Engine/Source/Runtime/AIModule/Classes/VisualLoggerExtension.h
+++ b/Engine/Source/Runtime/AIModule/Classes/VisualLoggerExtension.h
@@ -42,8 +42,8 @@ protected:
 };
 #endif //ENABLE_VISUAL_LOG
 
-UCLASS(Abstract)
-class AIMODULE_API UVisualLoggerExtension : public UObject
+UCLASS(Abstract, MinimalAPI)
+class UVisualLoggerExtension : public UObject
 {
 	GENERATED_BODY()
 };
diff --git a/Engine/Source/Runtime/AIModule/Public/BehaviorTreeDelegates.h b/Engine/Source/Runtime/AIModule/Public/BehaviorTreeDelegates.h
--- a/Engine/Source/Runtime/AIModule/Public/BehaviorTreeDelegates.h
+++ b/Engine/Source/Runtime/AIModule/Public/BehaviorTreeDelegates.h
@@ -9,7 +9,7 @@ class APawn;
 class UBehaviorTree;
 class UBehaviorTreeComponent;
 
-struct AIMODULE_API FBehaviorTreeDelegates
+struct FBehaviorTreeDelegates
 {
 	/** delegate type for tree execution events (Params: const UBehaviorTreeComponent* OwnerComp, const UBehaviorTree* TreeAsset) */
 	DECLARE_MULTICAST_DELEGATE_TwoParams(FOnTreeStarted, const UBehaviorTreeComponent&, const UBehaviorTree& );
@@ -18,11 +18,11 @@ struct AIMODULE_API FBehaviorTreeDelegates
 	DECLARE_MULTICAST_DELEGATE_OneParam(FOnTreeDebugTool, const APawn*);
 
 	/** Called when the behavior tree starts execution. */
-	static FOnTreeStarted OnTreeStarted;
+	static AIMODULE_API FOnTreeStarted OnTreeStarted;
 
 	/** Called when the AI debug tool highlights a pawn. */
-	static FOnTreeDebugTool OnDebugSelected;
+	static AIMODULE_API FOnTreeDebugTool OnDebugSelected;
 
 	/** Called when the AI debug tool locks on pawn. */
-	static FOnTreeDebugTool OnDebugLocked;
+	static AIMODULE_API FOnTreeDebugTool OnDebugLocked;
 };
diff --git a/Engine/Source/Runtime/AIModule/Public/IndexedHandle.h b/Engine/Source/Runtime/AIModule/Public/IndexedHandle.h
--- a/Engine/Source/Runtime/AIModule/Public/IndexedHandle.h
+++ b/Engine/Source/Runtime/AIModule/Public/IndexedHandle.h
@@ -19,7 +19,7 @@
  * For fast access in to index based data structures when we don't expect handles to be able to be stale
  */
 USTRUCT()
-struct AIMODULE_API FSimpleIndexedHandleBase
+struct FSimpleIndexedHandleBase
 {
 	GENERATED_BODY()
 
@@ -73,7 +73,7 @@ protected:
 
 /** Index based handle that has a serial number to verify stale handles. For fast safe access in to index based data structures */
 USTRUCT()
-struct AIMODULE_API FIndexedHandleBase : public FSimpleIndexedHandleBase
+struct FIndexedHandleBase : public FSimpleIndexedHandleBase
 {
 	GENERATED_BODY()
 
@@ -113,7 +113,7 @@ public:
 
 /** Compact Index based handle that has a serial number to verify stale handles. For fast safe access in to index based data structures. */
 USTRUCT()
-struct AIMODULE_API FCompactIndexedHandleBase
+struct FCompactIndexedHandleBase
 {
 	GENERATED_BODY()
 
diff --git a/Engine/Source/Runtime/AIModule/Public/SequentialID.h b/Engine/Source/Runtime/AIModule/Public/SequentialID.h
--- a/Engine/Source/Runtime/AIModule/Public/SequentialID.h
+++ b/Engine/Source/Runtime/AIModule/Public/SequentialID.h
@@ -5,7 +5,7 @@
 #include "SequentialID.generated.h"
 
 USTRUCT()
-struct AIMODULE_API FSequentialIDBase
+struct FSequentialIDBase
 {
 	GENERATED_BODY()
 
@@ -32,4 +32,4 @@ protected:
 
 	UPROPERTY()
 	uint32 Value = InvalidID;
-};
\ No newline at end of file
+};
diff --git a/Engine/Source/Runtime/AVEncoder/Public/AudioEncoder.h b/Engine/Source/Runtime/AVEncoder/Public/AudioEncoder.h
--- a/Engine/Source/Runtime/AVEncoder/Public/AudioEncoder.h
+++ b/Engine/Source/Runtime/AVEncoder/Public/AudioEncoder.h
@@ -25,7 +25,7 @@ namespace AVEncoder
 		virtual void OnEncodedAudioFrame(const FMediaPacket& Packet) = 0;
 	};
 
-	class AVENCODER_API FAudioEncoder
+	class FAudioEncoder
 	{
 	public:
 		virtual ~FAudioEncoder() {}
@@ -41,11 +41,11 @@ namespace AVEncoder
 		virtual void Encode(const FAudioFrame& Frame) = 0;
 		virtual FAudioConfig GetConfig() const = 0;
 
-		virtual void RegisterListener(IAudioEncoderListener& Listener);
-		virtual void UnregisterListener(IAudioEncoderListener& Listener);
+		AVENCODER_API virtual void RegisterListener(IAudioEncoderListener& Listener);
+		AVENCODER_API virtual void UnregisterListener(IAudioEncoderListener& Listener);
 
 	protected:
-		void OnEncodedAudioFrame(const FMediaPacket& Packet);
+		AVENCODER_API void OnEncodedAudioFrame(const FMediaPacket& Packet);
 	private:
 
 		FCriticalSection ListenersMutex;
diff --git a/Engine/Source/Runtime/AVEncoder/Public/AudioEncoderFactory.h b/Engine/Source/Runtime/AVEncoder/Public/AudioEncoderFactory.h
--- a/Engine/Source/Runtime/AVEncoder/Public/AudioEncoderFactory.h
+++ b/Engine/Source/Runtime/AVEncoder/Public/AudioEncoderFactory.h
@@ -11,7 +11,7 @@ namespace AVEncoder
 {
 	class FAudioEncoder;
 
-	class AVENCODER_API FAudioEncoderFactory
+	class FAudioEncoderFactory
 	{
 	public:
 		virtual ~FAudioEncoderFactory() {}
@@ -19,12 +19,12 @@ namespace AVEncoder
 		virtual TArray<FString> GetSupportedCodecs() const = 0;
 		virtual TUniquePtr<FAudioEncoder> CreateEncoder(const FString& Codec) = 0;
 
-		static void RegisterFactory(FAudioEncoderFactory& Factory);
-		static void UnregisterFactory(FAudioEncoderFactory& Factory);
-		static FAudioEncoderFactory* FindFactory(const FString& Codec);
-		static const TArray<FAudioEncoderFactory*> GetAllFactories();
+		static AVENCODER_API void RegisterFactory(FAudioEncoderFactory& Factory);
+		static AVENCODER_API void UnregisterFactory(FAudioEncoderFactory& Factory);
+		static AVENCODER_API FAudioEncoderFactory* FindFactory(const FString& Codec);
+		static AVENCODER_API const TArray<FAudioEncoderFactory*> GetAllFactories();
 
 	private:
-		static TArray<FAudioEncoderFactory*> Factories;
+		static AVENCODER_API TArray<FAudioEncoderFactory*> Factories;
 	};
 }
diff --git a/Engine/Source/Runtime/AVEncoder/Public/CodecPacket.h b/Engine/Source/Runtime/AVEncoder/Public/CodecPacket.h
--- a/Engine/Source/Runtime/AVEncoder/Public/CodecPacket.h
+++ b/Engine/Source/Runtime/AVEncoder/Public/CodecPacket.h
@@ -9,12 +9,12 @@
 
 namespace AVEncoder
 {
-	class AVENCODER_API FCodecPacket
+	class FCodecPacket
 	{
 	public:
 		virtual ~FCodecPacket() = default;
 
-		static FCodecPacket Create(const uint8* InData, uint32 InDataSize);
+		static AVENCODER_API FCodecPacket Create(const uint8* InData, uint32 InDataSize);
 
 		/**
 		 * Encoding/Decoding latency
diff --git a/Engine/Source/Runtime/AVEncoder/Public/IMFSampleWrapper.h b/Engine/Source/Runtime/AVEncoder/Public/IMFSampleWrapper.h
--- a/Engine/Source/Runtime/AVEncoder/Public/IMFSampleWrapper.h
+++ b/Engine/Source/Runtime/AVEncoder/Public/IMFSampleWrapper.h
@@ -15,7 +15,7 @@ namespace AVEncoder
 //
 // Wrapper for IMFSample, to make it easier to report errors
 //
-class AVENCODER_API FIMFSampleWrapper
+class FIMFSampleWrapper
 {
 public:
 
@@ -35,19 +35,19 @@ public:
 		return Sample;
 	}
 
-	bool CreateSample();
+	AVENCODER_API bool CreateSample();
 
-	FTimespan GetTime() const;
+	AVENCODER_API FTimespan GetTime() const;
 
-	void SetTime(FTimespan Time);
+	AVENCODER_API void SetTime(FTimespan Time);
 
-	FTimespan GetDuration() const;
+	AVENCODER_API FTimespan GetDuration() const;
 
-	void SetDuration(FTimespan Duration);
+	AVENCODER_API void SetDuration(FTimespan Duration);
 
-	bool IsVideoKeyFrame() const;
+	AVENCODER_API bool IsVideoKeyFrame() const;
 
-	int GetBufferCount() const;
+	AVENCODER_API int GetBufferCount() const;
 
 	/**
 	*
@@ -92,7 +92,7 @@ public:
 		Sample = nullptr;
 	}
 
-	FIMFSampleWrapper Clone() const;
+	AVENCODER_API FIMFSampleWrapper Clone() const;
 
 private:
 	EPacketType MediaType;
diff --git a/Engine/Source/Runtime/AVEncoder/Public/VideoDecoderFactory.h b/Engine/Source/Runtime/AVEncoder/Public/VideoDecoderFactory.h
--- a/Engine/Source/Runtime/AVEncoder/Public/VideoDecoderFactory.h
+++ b/Engine/Source/Runtime/AVEncoder/Public/VideoDecoderFactory.h
@@ -15,11 +15,11 @@ namespace AVEncoder
 {
 	struct FVideoDecoderInfo;
 
-	class AVENCODER_API FVideoDecoderFactory
+	class FVideoDecoderFactory
 	{
 	public:
-		static FVideoDecoderFactory& Get();
-		static void Debug_SetDontRegisterDefaultCodecs();
+		static AVENCODER_API FVideoDecoderFactory& Get();
+		static AVENCODER_API void Debug_SetDontRegisterDefaultCodecs();
 
 		// --- decoder registry
 
@@ -27,29 +27,29 @@ namespace AVEncoder
 		using CreateDecoderCallback = TFunction<FVideoDecoder*()>;
 
 		// register a decoder so that it can be iterated and created
-		void Register(const FVideoDecoderInfo& InInfo, const CreateDecoderCallback& InCreateDecoder);
+		AVENCODER_API void Register(const FVideoDecoderInfo& InInfo, const CreateDecoderCallback& InCreateDecoder);
 
 		// get a list of registered decoders
 		const TArray<FVideoDecoderInfo>& GetAvailable() const { return AvailableDecoders; }
-		bool GetInfo(uint32 InID, FVideoDecoderInfo& OutInfo) const;
+		AVENCODER_API bool GetInfo(uint32 InID, FVideoDecoderInfo& OutInfo) const;
 
 		// --- decoder creation
 		// create a decoder instance. If not nullptr it must be destroyed through invoking its Shutdown() method!
-		FVideoDecoder* Create(uint32 InID, const FVideoDecoder::FInit& InInit);
+		AVENCODER_API FVideoDecoder* Create(uint32 InID, const FVideoDecoder::FInit& InInit);
 
 	private:
-		FVideoDecoderFactory();
-		~FVideoDecoderFactory();
+		AVENCODER_API FVideoDecoderFactory();
+		AVENCODER_API ~FVideoDecoderFactory();
 		FVideoDecoderFactory(const FVideoDecoderFactory&) = delete;
 		FVideoDecoderFactory& operator=(const FVideoDecoderFactory&) = delete;
 
-		void RegisterDefaultCodecs();
+		AVENCODER_API void RegisterDefaultCodecs();
 
-		static FThreadSafeCounter		NextID;
-		static bool						bDebugDontRegisterDefaultCodecs;
+		static AVENCODER_API FThreadSafeCounter		NextID;
+		static AVENCODER_API bool						bDebugDontRegisterDefaultCodecs;
 
 		TArray<FVideoDecoderInfo>		AvailableDecoders;
 		TArray<CreateDecoderCallback>	CreateDecoders;
 	};
 
-} /* namespace AVEncoder */
\ No newline at end of file
+} /* namespace AVEncoder */
diff --git a/Engine/Source/Runtime/AVEncoder/Public/VideoDecoderInput.h b/Engine/Source/Runtime/AVEncoder/Public/VideoDecoderInput.h
--- a/Engine/Source/Runtime/AVEncoder/Public/VideoDecoderInput.h
+++ b/Engine/Source/Runtime/AVEncoder/Public/VideoDecoderInput.h
@@ -14,7 +14,7 @@ struct FVideoDecoderAllocFrameBufferResult;
 namespace AVEncoder
 {
 
-	class AVENCODER_API FVideoDecoderInput
+	class FVideoDecoderInput
 	{
 	public:
 
@@ -32,7 +32,7 @@ namespace AVEncoder
 			bool		bMissingFrames = false;
 		};
 
-		static TSharedPtr<FVideoDecoderInput> Create(const FInputData& InInputData);
+		static AVENCODER_API TSharedPtr<FVideoDecoderInput> Create(const FInputData& InInputData);
 
 		virtual int32 GetWidth() const = 0;
 		virtual int32 GetHeight() const = 0;
@@ -54,7 +54,7 @@ namespace AVEncoder
 
 
 
-	class AVENCODER_API FVideoDecoderOutput
+	class FVideoDecoderOutput
 	{
 	public:
 		virtual int32 AddRef() = 0;
diff --git a/Engine/Source/Runtime/AVEncoder/Public/VideoEncoder.h b/Engine/Source/Runtime/AVEncoder/Public/VideoEncoder.h
--- a/Engine/Source/Runtime/AVEncoder/Public/VideoEncoder.h
+++ b/Engine/Source/Runtime/AVEncoder/Public/VideoEncoder.h
@@ -19,7 +19,7 @@ namespace AVEncoder
 	class FVideoEncoderInput;
 	class FVideoEncoderInputFrame;
 
-    class AVENCODER_API FVideoEncoder
+    class FVideoEncoder
     {
     public:
         enum class RateControlMode { UNKNOWN, CONSTQP, VBR, CBR };
@@ -61,17 +61,17 @@ namespace AVEncoder
 			}
         };
 
-        virtual ~FVideoEncoder();
+        AVENCODER_API virtual ~FVideoEncoder();
 
 		virtual bool Setup(TSharedRef<FVideoEncoderInput> input, FLayerConfig const& config) { return false; }
 		virtual void Shutdown() {}
 
-        virtual bool AddLayer(FLayerConfig const& config);
+        AVENCODER_API virtual bool AddLayer(FLayerConfig const& config);
         uint32 GetNumLayers() const { return static_cast<uint32>(Layers.Num()); }
         virtual uint32 GetMaxLayers() const { return 1; }
 
-		FLayerConfig GetLayerConfig(uint32 layerIdx) const;
-        void UpdateLayerConfig(uint32 layerIdx, FLayerConfig const& config);
+		AVENCODER_API FLayerConfig GetLayerConfig(uint32 layerIdx) const;
+        AVENCODER_API void UpdateLayerConfig(uint32 layerIdx, FLayerConfig const& config);
 
         using OnFrameEncodedCallback = TFunction<void(const TSharedPtr<FVideoEncoderInputFrame> /* InCompletedFrame */)>;
         using OnEncodedPacketCallback = TFunction<void(uint32 /* LayerIndex */, const TSharedPtr<FVideoEncoderInputFrame> /* Frame */, const FCodecPacket& /* Packet */)>;
@@ -119,4 +119,4 @@ namespace AVEncoder
         TArray<FLayer*>			Layers;
         OnEncodedPacketCallback	OnEncodedPacket;
     };
-}
\ No newline at end of file
+}
diff --git a/Engine/Source/Runtime/AVEncoder/Public/VideoEncoderFactory.h b/Engine/Source/Runtime/AVEncoder/Public/VideoEncoderFactory.h
--- a/Engine/Source/Runtime/AVEncoder/Public/VideoEncoderFactory.h
+++ b/Engine/Source/Runtime/AVEncoder/Public/VideoEncoderFactory.h
@@ -18,12 +18,12 @@ namespace AVEncoder
 {
 	class FVideoEncoderInput;
 
-	class AVENCODER_API FVideoEncoderFactory
+	class FVideoEncoderFactory
 	{
 	public:
-		static FVideoEncoderFactory& Get();
-		static void Shutdown();
-		static void Debug_SetDontRegisterDefaultCodecs();
+		static AVENCODER_API FVideoEncoderFactory& Get();
+		static AVENCODER_API void Shutdown();
+		static AVENCODER_API void Debug_SetDontRegisterDefaultCodecs();
 
 		bool IsSetup() const { return bWasSetup; }
 
@@ -33,17 +33,17 @@ namespace AVEncoder
 		using CreateEncoderCallback = TFunction<TUniquePtr<FVideoEncoder>()>;
 
 		// register an encoder so that it can be iterated and created
-		void Register(const FVideoEncoderInfo& InInfo, const CreateEncoderCallback& InCreateEncoder);
+		AVENCODER_API void Register(const FVideoEncoderInfo& InInfo, const CreateEncoderCallback& InCreateEncoder);
 
 		// get a list of registered encoders
 		const TArray<FVideoEncoderInfo>& GetAvailable() const { return AvailableEncoders; }
-		bool GetInfo(uint32 InID, FVideoEncoderInfo& OutInfo) const;
-		bool HasEncoderForCodec(ECodecType CodecType) const;
+		AVENCODER_API bool GetInfo(uint32 InID, FVideoEncoderInfo& OutInfo) const;
+		AVENCODER_API bool HasEncoderForCodec(ECodecType CodecType) const;
 
 		// --- encoder creation
 
-		TUniquePtr<FVideoEncoder> Create(uint32 InID, const FVideoEncoder::FLayerConfig& config);
-		TUniquePtr<FVideoEncoder> Create(uint32 InID, TSharedPtr<FVideoEncoderInput> InInput, const FVideoEncoder::FLayerConfig& config);
+		AVENCODER_API TUniquePtr<FVideoEncoder> Create(uint32 InID, const FVideoEncoder::FLayerConfig& config);
+		AVENCODER_API TUniquePtr<FVideoEncoder> Create(uint32 InID, TSharedPtr<FVideoEncoderInput> InInput, const FVideoEncoder::FLayerConfig& config);
 
 	private:
 		FVideoEncoderFactory() = default;
@@ -51,11 +51,11 @@ namespace AVEncoder
 		FVideoEncoderFactory(const FVideoEncoderFactory&) = delete;
 		FVideoEncoderFactory& operator=(const FVideoEncoderFactory&) = delete;
 
-		void RegisterDefaultCodecs();
+		AVENCODER_API void RegisterDefaultCodecs();
 
-		static FCriticalSection			ProtectSingleton;
-		static FVideoEncoderFactory		Singleton;
-		static FThreadSafeCounter		NextID;
+		static AVENCODER_API FCriticalSection			ProtectSingleton;
+		static AVENCODER_API FVideoEncoderFactory		Singleton;
+		static AVENCODER_API FThreadSafeCounter		NextID;
 		bool							bDebugDontRegisterDefaultCodecs = false;
 		FThreadSafeBool					bWasSetup;
 		TArray<FVideoEncoderInfo>		AvailableEncoders;
@@ -63,4 +63,4 @@ namespace AVEncoder
 
 	};
     
-} /* namespace AVEncoder */
\ No newline at end of file
+} /* namespace AVEncoder */
diff --git a/Engine/Source/Runtime/AVEncoder/Public/VideoEncoderInput.h b/Engine/Source/Runtime/AVEncoder/Public/VideoEncoderInput.h
--- a/Engine/Source/Runtime/AVEncoder/Public/VideoEncoderInput.h
+++ b/Engine/Source/Runtime/AVEncoder/Public/VideoEncoderInput.h
@@ -42,27 +42,27 @@ namespace AVEncoder
 	};
 #endif
 
-	class AVENCODER_API FVideoEncoderInput
+	class FVideoEncoderInput
 	{
 	public:
 		// --- construct video encoder input based on expected input frame format
-		static TSharedPtr<FVideoEncoderInput> CreateDummy(bool isResizable = false);
-		static TSharedPtr<FVideoEncoderInput> CreateForYUV420P(uint32 InWidth, uint32 InHeight, bool isResizable = false);
+		static AVENCODER_API TSharedPtr<FVideoEncoderInput> CreateDummy(bool isResizable = false);
+		static AVENCODER_API TSharedPtr<FVideoEncoderInput> CreateForYUV420P(uint32 InWidth, uint32 InHeight, bool isResizable = false);
 
 		// create input for an encoder that encodes a D3D11 texture 
-		static TSharedPtr<FVideoEncoderInput> CreateForD3D11(void* InApplicationD3D11Device, bool IsResizable = false, bool IsShared = false);
+		static AVENCODER_API TSharedPtr<FVideoEncoderInput> CreateForD3D11(void* InApplicationD3D11Device, bool IsResizable = false, bool IsShared = false);
 
 		// create input for an encoder that encodes a D3D12 texture
-		static TSharedPtr<FVideoEncoderInput> CreateForD3D12(void* InApplicationD3D12Device, bool IsResizable = false, bool IsShared = false);
+		static AVENCODER_API TSharedPtr<FVideoEncoderInput> CreateForD3D12(void* InApplicationD3D12Device, bool IsResizable = false, bool IsShared = false);
 
 		// create input for an encoder that encodes a CUarray
-		static TSharedPtr<FVideoEncoderInput> CreateForCUDA(void* InApplicationCudaContext, bool IsResizable = false);
+		static AVENCODER_API TSharedPtr<FVideoEncoderInput> CreateForCUDA(void* InApplicationCudaContext, bool IsResizable = false);
 
 		// create input for an encoder that encodes a VkImage
-		static TSharedPtr<FVideoEncoderInput> CreateForVulkan(void* InApplicationVulkanData, bool IsResizable = false);
+		static AVENCODER_API TSharedPtr<FVideoEncoderInput> CreateForVulkan(void* InApplicationVulkanData, bool IsResizable = false);
 
 		// --- properties
-		virtual void SetMaxNumBuffers(uint32 InMaxNumBuffers);
+		AVENCODER_API virtual void SetMaxNumBuffers(uint32 InMaxNumBuffers);
 
 		EVideoFrameFormat GetFrameFormat() const { return FrameFormat; }
 
@@ -123,7 +123,7 @@ namespace AVEncoder
 
 
 	// TODO this should go elsewhere and be made cross platform
-	class AVENCODER_API FVideoEncoderInputFrame
+	class FVideoEncoderInputFrame
 	{
 	public:
 		// Obtain (increase reference count) of this input frame
@@ -167,7 +167,7 @@ namespace AVEncoder
 			uint32				StrideV = 0;
 		};
 
-		void AllocateYUV420P();
+		AVENCODER_API void AllocateYUV420P();
 		const FYUV420P& GetYUV420P() const 
 		{	
 			return YUV420P; 
@@ -175,7 +175,7 @@ namespace AVEncoder
 		
 		FYUV420P& GetYUV420P() { return YUV420P; }
 
-		void SetYUV420P(const uint8* InDataY, const uint8* InDataU, const uint8* InDataV, uint32 InStrideY, uint32 InStrideU, uint32 InStrideV);
+		AVENCODER_API void SetYUV420P(const uint8* InDataY, const uint8* InDataU, const uint8* InDataV, uint32 InStrideY, uint32 InStrideU, uint32 InStrideV);
 
 #if PLATFORM_WINDOWS
 		// --- D3D11
@@ -194,7 +194,7 @@ namespace AVEncoder
 		// the callback type used to create a registered encoder
 		using FReleaseD3D11TextureCallback = TFunction<void(ID3D11Texture2D*)>;
 
-		void SetTexture(ID3D11Texture2D* InTexture, FReleaseD3D11TextureCallback InOnReleaseD3D11Texture);
+		AVENCODER_API void SetTexture(ID3D11Texture2D* InTexture, FReleaseD3D11TextureCallback InOnReleaseD3D11Texture);
 
 		// --- D3D12
 
@@ -211,7 +211,7 @@ namespace AVEncoder
 		// the callback type used to create a registered encoder
 		using FReleaseD3D12TextureCallback = TFunction<void(ID3D12Resource*)>;
 
-		void SetTexture(ID3D12Resource* InTexture, FReleaseD3D12TextureCallback InOnReleaseD3D11Texture);
+		AVENCODER_API void SetTexture(ID3D12Resource* InTexture, FReleaseD3D12TextureCallback InOnReleaseD3D11Texture);
 
 #endif // PLATFORM_WINDOWS
 
@@ -238,7 +238,7 @@ namespace AVEncoder
 		// the callback type used to create a registered encoder
 		using FReleaseCUDATextureCallback = TFunction<void(CUarray)>;
 
-		void SetTexture(CUarray InTexture, EUnderlyingRHI UnderlyingRHI, void* SharedHandle, FReleaseCUDATextureCallback InOnReleaseTexture);
+		AVENCODER_API void SetTexture(CUarray InTexture, EUnderlyingRHI UnderlyingRHI, void* SharedHandle, FReleaseCUDATextureCallback InOnReleaseTexture);
 
 #if PLATFORM_DESKTOP && !PLATFORM_APPLE
 		// --- Vulkan
@@ -259,14 +259,14 @@ namespace AVEncoder
 		using FReleaseVulkanSurfaceCallback = TFunction<void(void*)>;
 		mutable FReleaseVulkanSurfaceCallback OnReleaseVulkanSurface;
 
-		void SetTexture(VkImage InTexture, FReleaseVulkanTextureCallback InOnReleaseTexture);
-		void SetTexture(VkImage InTexture, VkDeviceMemory InTextureDeviceMemory, uint64 InTextureSize, FReleaseVulkanTextureCallback InOnReleaseTexture);
+		AVENCODER_API void SetTexture(VkImage InTexture, FReleaseVulkanTextureCallback InOnReleaseTexture);
+		AVENCODER_API void SetTexture(VkImage InTexture, VkDeviceMemory InTextureDeviceMemory, uint64 InTextureSize, FReleaseVulkanTextureCallback InOnReleaseTexture);
 #endif
 
-		virtual ~FVideoEncoderInputFrame();
+		AVENCODER_API virtual ~FVideoEncoderInputFrame();
 	protected:
-		FVideoEncoderInputFrame();
-		explicit FVideoEncoderInputFrame(const FVideoEncoderInputFrame& CloneFrom);
+		AVENCODER_API FVideoEncoderInputFrame();
+		AVENCODER_API explicit FVideoEncoderInputFrame(const FVideoEncoderInputFrame& CloneFrom);
 		
 
 		uint32									FrameID;
diff --git a/Engine/Source/Runtime/AdvancedWidgets/Public/AdvancedWidgetsModule.h b/Engine/Source/Runtime/AdvancedWidgets/Public/AdvancedWidgetsModule.h
--- a/Engine/Source/Runtime/AdvancedWidgets/Public/AdvancedWidgetsModule.h
+++ b/Engine/Source/Runtime/AdvancedWidgets/Public/AdvancedWidgetsModule.h
@@ -15,24 +15,24 @@ class FPropertyValueFactory;
 /**
  * Advanced Widgets module
  */
-class ADVANCEDWIDGETS_API FAdvancedWidgetsModule : public IModuleInterface
+class FAdvancedWidgetsModule : public IModuleInterface
 {
 public:
-	static FAdvancedWidgetsModule& GetModule();
-	static FAdvancedWidgetsModule* GetModulePtr();
+	static ADVANCEDWIDGETS_API FAdvancedWidgetsModule& GetModule();
+	static ADVANCEDWIDGETS_API FAdvancedWidgetsModule* GetModulePtr();
 
 	/**
 	 * Called right after the plugin DLL has been loaded and the plugin object has been created
 	 */
-	virtual void StartupModule();
+	ADVANCEDWIDGETS_API virtual void StartupModule();
 
 	/**
 	 * Called before the plugin is unloaded, right before the plugin object is destroyed.
 	 */
-	virtual void ShutdownModule();
+	ADVANCEDWIDGETS_API virtual void ShutdownModule();
 
 public:
-	UE::PropertyViewer::FPropertyValueFactory& GetPropertyValueFactory() const;
+	ADVANCEDWIDGETS_API UE::PropertyViewer::FPropertyValueFactory& GetPropertyValueFactory() const;
 
 private:
 	TUniquePtr<UE::PropertyViewer::FPropertyValueFactory> PropertyValueFactory;
diff --git a/Engine/Source/Runtime/AdvancedWidgets/Public/Components/RadialSlider.h b/Engine/Source/Runtime/AdvancedWidgets/Public/Components/RadialSlider.h
--- a/Engine/Source/Runtime/AdvancedWidgets/Public/Components/RadialSlider.h
+++ b/Engine/Source/Runtime/AdvancedWidgets/Public/Components/RadialSlider.h
@@ -23,8 +23,8 @@ class SRadialSlider;
  * * No Children
  */
 
-UCLASS()
-class ADVANCEDWIDGETS_API URadialSlider : public UWidget
+UCLASS(MinimalAPI)
+class URadialSlider : public UWidget
 {
 	GENERATED_UCLASS_BODY()
 
@@ -143,94 +143,94 @@ public:
 
 	/** Gets the current value of the slider. */
 	UFUNCTION(BlueprintCallable, Category="Behavior")
-	float GetValue() const;
+	ADVANCEDWIDGETS_API float GetValue() const;
 
 	/** Gets the current custom default value of the slider. */
 	UFUNCTION(BlueprintCallable, Category = "Behavior")
-	float GetCustomDefaultValue() const;
+	ADVANCEDWIDGETS_API float GetCustomDefaultValue() const;
 
 	/** Get the current raw slider alpha from 0 to 1 */
 	UFUNCTION(BlueprintCallable, Category = "Behavior")
-	float GetNormalizedSliderHandlePosition() const;
+	ADVANCEDWIDGETS_API float GetNormalizedSliderHandlePosition() const;
 
 	/** Sets the current value of the slider. */
 	UFUNCTION(BlueprintCallable, Category="Behavior")
-	void SetValue(float InValue);
+	ADVANCEDWIDGETS_API void SetValue(float InValue);
 
 	/** Sets the current custom default value of the slider. */
 	UFUNCTION(BlueprintCallable, Category = "Behavior")
-	void SetCustomDefaultValue(float InValue);
+	ADVANCEDWIDGETS_API void SetCustomDefaultValue(float InValue);
 
 	/** Sets the curve for the slider range*/
 	UFUNCTION(BlueprintCallable, Category = "Behaviour")
-	void SetSliderRange(const FRuntimeFloatCurve& InSliderRange);
+	ADVANCEDWIDGETS_API void SetSliderRange(const FRuntimeFloatCurve& InSliderRange);
 
 	/** Adds value tags to the slider. */
 	UFUNCTION(BlueprintCallable, Category = "Behaviour")
-	void SetValueTags(const TArray<float>& InValueTags);
+	ADVANCEDWIDGETS_API void SetValueTags(const TArray<float>& InValueTags);
 
 	/** Sets the minimum angle of the slider. */
 	UFUNCTION(BlueprintCallable, Category = "Behavior")
-	void SetSliderHandleStartAngle(float InValue);
+	ADVANCEDWIDGETS_API void SetSliderHandleStartAngle(float InValue);
 
 	/** Sets the maximum angle of the slider. */
 	UFUNCTION(BlueprintCallable, Category = "Behavior")
-	void SetSliderHandleEndAngle(float InValue);
+	ADVANCEDWIDGETS_API void SetSliderHandleEndAngle(float InValue);
 
 	/** Sets the Angular Offset for the slider. */
 	UFUNCTION(BlueprintCallable, Category = "Behaviour")
-	void SetAngularOffset(float InValue);
+	ADVANCEDWIDGETS_API void SetAngularOffset(float InValue);
 
 	/** Sets the start and end of the hand as a ratio to the slider radius (so 0.0 to 1.0 is from the slider center to the handle). */
 	UFUNCTION(BlueprintCallable, Category = "Appearance")
-	void SetHandStartEndRatio(FVector2D InValue);
+	ADVANCEDWIDGETS_API void SetHandStartEndRatio(FVector2D InValue);
 
 	/** Sets the handle to be interactive or fixed */
 	UFUNCTION(BlueprintCallable, Category="Behavior")
-	void SetLocked(bool InValue);
+	ADVANCEDWIDGETS_API void SetLocked(bool InValue);
 
 	/** Sets the amount to adjust the value by, when using a controller or keyboard */
 	UFUNCTION(BlueprintCallable, Category="Behavior")
-	void SetStepSize(float InValue);
+	ADVANCEDWIDGETS_API void SetStepSize(float InValue);
 
 	/** Sets the color of the slider bar */
 	UFUNCTION(BlueprintCallable, Category="Appearance")
-	void SetSliderBarColor(FLinearColor InValue);
+	ADVANCEDWIDGETS_API void SetSliderBarColor(FLinearColor InValue);
 
 	/** Sets the progress color of the slider bar */
 	UFUNCTION(BlueprintCallable, Category = "Appearance")
-	void SetSliderProgressColor(FLinearColor InValue);
+	ADVANCEDWIDGETS_API void SetSliderProgressColor(FLinearColor InValue);
 
 	/** Sets the color of the handle bar */
 	UFUNCTION(BlueprintCallable, Category="Appearance")
-	void SetSliderHandleColor(FLinearColor InValue);
+	ADVANCEDWIDGETS_API void SetSliderHandleColor(FLinearColor InValue);
 
 	/** Sets the color of the slider bar */
 	UFUNCTION(BlueprintCallable, Category="Appearance")
-	void SetCenterBackgroundColor(FLinearColor InValue);
+	ADVANCEDWIDGETS_API void SetCenterBackgroundColor(FLinearColor InValue);
 
 	/** Set whether the value is changed when dragging vertically as opposed to along the radial curve.  */
 	UFUNCTION(BlueprintCallable, Category = "Behavior")
-	void SetUseVerticalDrag(bool InUseVerticalDrag);
+	ADVANCEDWIDGETS_API void SetUseVerticalDrag(bool InUseVerticalDrag);
 
 	/** Whether to show the slider handle (thumb). */
 	UFUNCTION(BlueprintCallable, Category = "Appearance")
-	void SetShowSliderHandle(bool InShowSliderHandle);
+	ADVANCEDWIDGETS_API void SetShowSliderHandle(bool InShowSliderHandle);
 
 	/** Whether to show the slider hand. */
 	UFUNCTION(BlueprintCallable, Category = "Appearance")
-	void SetShowSliderHand(bool InShowSliderHand);
+	ADVANCEDWIDGETS_API void SetShowSliderHand(bool InShowSliderHand);
 	
 	// UWidget interface
-	virtual void SynchronizeProperties() override;
+	ADVANCEDWIDGETS_API virtual void SynchronizeProperties() override;
 	// End of UWidget interface
 
 	// UVisual interface
-	virtual void ReleaseSlateResources(bool bReleaseChildren) override;
+	ADVANCEDWIDGETS_API virtual void ReleaseSlateResources(bool bReleaseChildren) override;
 	// End of UVisual interface
 
 #if WITH_EDITOR
-	virtual const FText GetPaletteCategory() override;
+	ADVANCEDWIDGETS_API virtual const FText GetPaletteCategory() override;
 #endif
 
 protected:
@@ -238,17 +238,17 @@ protected:
 	TSharedPtr<SRadialSlider> MyRadialSlider;
 
 	// UWidget interface
-	virtual TSharedRef<SWidget> RebuildWidget() override;
+	ADVANCEDWIDGETS_API virtual TSharedRef<SWidget> RebuildWidget() override;
 	// End of UWidget interface
 
-	void HandleOnValueChanged(float InValue);
-	void HandleOnMouseCaptureBegin();
-	void HandleOnMouseCaptureEnd();
-	void HandleOnControllerCaptureBegin();
-	void HandleOnControllerCaptureEnd();
+	ADVANCEDWIDGETS_API void HandleOnValueChanged(float InValue);
+	ADVANCEDWIDGETS_API void HandleOnMouseCaptureBegin();
+	ADVANCEDWIDGETS_API void HandleOnMouseCaptureEnd();
+	ADVANCEDWIDGETS_API void HandleOnControllerCaptureBegin();
+	ADVANCEDWIDGETS_API void HandleOnControllerCaptureEnd();
 
 #if WITH_ACCESSIBILITY
-	virtual TSharedPtr<SWidget> GetAccessibleWidget() const override;
+	ADVANCEDWIDGETS_API virtual TSharedPtr<SWidget> GetAccessibleWidget() const override;
 #endif
 
 	PROPERTY_BINDING_IMPLEMENTATION(float, Value);
diff --git a/Engine/Source/Runtime/AdvancedWidgets/Public/Framework/PropertyViewer/FieldIconFinder.h b/Engine/Source/Runtime/AdvancedWidgets/Public/Framework/PropertyViewer/FieldIconFinder.h
--- a/Engine/Source/Runtime/AdvancedWidgets/Public/Framework/PropertyViewer/FieldIconFinder.h
+++ b/Engine/Source/Runtime/AdvancedWidgets/Public/Framework/PropertyViewer/FieldIconFinder.h
@@ -62,7 +62,7 @@ struct FFieldColorSettings
 
 
 /** */
-class ADVANCEDWIDGETS_API FFieldIconFinder
+class FFieldIconFinder
 {
 public:
 	/** */
@@ -75,15 +75,15 @@ public:
 	using FFieldIconArray = TArray<FFieldIcon, TInlineAllocator<2>>;
 
 	/** */
-	static const FSlateBrush* GetIcon(const UObject* Object);
+	static ADVANCEDWIDGETS_API const FSlateBrush* GetIcon(const UObject* Object);
 
 	/** */
-	static FFieldIconArray GetFunctionIcon(const UFunction* Function);
-	static FFieldIconArray GetFunctionIcon(const UFunction* Function, const FFieldColorSettings& Settings);
+	static ADVANCEDWIDGETS_API FFieldIconArray GetFunctionIcon(const UFunction* Function);
+	static ADVANCEDWIDGETS_API FFieldIconArray GetFunctionIcon(const UFunction* Function, const FFieldColorSettings& Settings);
 
 	/** */
-	static FFieldIconArray GetPropertyIcon(const FProperty* Property);
-	static FFieldIconArray GetPropertyIcon(const FProperty* Property, const FFieldColorSettings& Settings);
+	static ADVANCEDWIDGETS_API FFieldIconArray GetPropertyIcon(const FProperty* Property);
+	static ADVANCEDWIDGETS_API FFieldIconArray GetPropertyIcon(const FProperty* Property, const FFieldColorSettings& Settings);
 };
 
 } //namespace
diff --git a/Engine/Source/Runtime/AdvancedWidgets/Public/Framework/PropertyViewer/IFieldExpander.h b/Engine/Source/Runtime/AdvancedWidgets/Public/Framework/PropertyViewer/IFieldExpander.h
--- a/Engine/Source/Runtime/AdvancedWidgets/Public/Framework/PropertyViewer/IFieldExpander.h
+++ b/Engine/Source/Runtime/AdvancedWidgets/Public/Framework/PropertyViewer/IFieldExpander.h
@@ -17,7 +17,7 @@ namespace UE::PropertyViewer
 {
 
 /** */
-class ADVANCEDWIDGETS_API IFieldExpander
+class IFieldExpander
 {
 public:
 	/** @return the class that should be expended for this instance. The instance can be null. */
@@ -35,7 +35,7 @@ public:
 	}
 };
 
-class ADVANCEDWIDGETS_API FFieldExpander_Default : public IFieldExpander
+class FFieldExpander_Default : public IFieldExpander
 {
 public:
 	FFieldExpander_Default() = default;
diff --git a/Engine/Source/Runtime/AdvancedWidgets/Public/Framework/PropertyViewer/IFieldIterator.h b/Engine/Source/Runtime/AdvancedWidgets/Public/Framework/PropertyViewer/IFieldIterator.h
--- a/Engine/Source/Runtime/AdvancedWidgets/Public/Framework/PropertyViewer/IFieldIterator.h
+++ b/Engine/Source/Runtime/AdvancedWidgets/Public/Framework/PropertyViewer/IFieldIterator.h
@@ -15,16 +15,16 @@ namespace UE::PropertyViewer
 {
 
 /** */
-class ADVANCEDWIDGETS_API IFieldIterator
+class IFieldIterator
 {
 public:
 	virtual TArray<FFieldVariant> GetFields(const UStruct*) const = 0;
 	virtual ~IFieldIterator() = default;
 };
 
-class ADVANCEDWIDGETS_API FFieldIterator_BlueprintVisible : public IFieldIterator
+class FFieldIterator_BlueprintVisible : public IFieldIterator
 {
-	virtual TArray<FFieldVariant> GetFields(const UStruct*) const override;
+	ADVANCEDWIDGETS_API virtual TArray<FFieldVariant> GetFields(const UStruct*) const override;
 };
 
 } //namespace
diff --git a/Engine/Source/Runtime/AdvancedWidgets/Public/Framework/PropertyViewer/PropertyPath.h b/Engine/Source/Runtime/AdvancedWidgets/Public/Framework/PropertyViewer/PropertyPath.h
--- a/Engine/Source/Runtime/AdvancedWidgets/Public/Framework/PropertyViewer/PropertyPath.h
+++ b/Engine/Source/Runtime/AdvancedWidgets/Public/Framework/PropertyViewer/PropertyPath.h
@@ -23,16 +23,16 @@ class UScriptStruct;
 namespace UE::PropertyViewer
 {
 
-struct ADVANCEDWIDGETS_API FPropertyPath
+struct FPropertyPath
 {
 public:
 	using FPropertyArray = TArray<const FProperty*, TInlineAllocator<1>>;
 
 	FPropertyPath() = default;
-	FPropertyPath(UObject* Object, const FProperty* Property);
-	FPropertyPath(UObject* Object, FPropertyArray Properties);
-	FPropertyPath(const UScriptStruct* ScriptStruct, void* Data, const FProperty* Property);
-	FPropertyPath(const UScriptStruct* ScriptStruct, void* Data, FPropertyArray Properties);
+	ADVANCEDWIDGETS_API FPropertyPath(UObject* Object, const FProperty* Property);
+	ADVANCEDWIDGETS_API FPropertyPath(UObject* Object, FPropertyArray Properties);
+	ADVANCEDWIDGETS_API FPropertyPath(const UScriptStruct* ScriptStruct, void* Data, const FProperty* Property);
+	ADVANCEDWIDGETS_API FPropertyPath(const UScriptStruct* ScriptStruct, void* Data, FPropertyArray Properties);
 
 public:
 	bool HasProperty() const
@@ -50,8 +50,8 @@ public:
 		return MakeArrayView(Properties);
 	}
 
-	void* GetContainerPtr();
-	const void* GetContainerPtr() const;
+	ADVANCEDWIDGETS_API void* GetContainerPtr();
+	ADVANCEDWIDGETS_API const void* GetContainerPtr() const;
 
 private:
 	TWeakObjectPtr<UObject> TopLevelContainer_Object = nullptr;
diff --git a/Engine/Source/Runtime/AdvancedWidgets/Public/Framework/PropertyViewer/PropertyValueFactory.h b/Engine/Source/Runtime/AdvancedWidgets/Public/Framework/PropertyViewer/PropertyValueFactory.h
--- a/Engine/Source/Runtime/AdvancedWidgets/Public/Framework/PropertyViewer/PropertyValueFactory.h
+++ b/Engine/Source/Runtime/AdvancedWidgets/Public/Framework/PropertyViewer/PropertyValueFactory.h
@@ -24,7 +24,7 @@ class FPropertyValueFactoryImpl;
 
 
 /** */
-class ADVANCEDWIDGETS_API FPropertyValueFactory
+class FPropertyValueFactory
 {
 public:
 	struct FHandle
@@ -58,22 +58,22 @@ public:
 	DECLARE_DELEGATE_RetVal_OneParam(TSharedPtr<SWidget>, FOnGenerate, FGenerateArgs);
 
 public:
-	FPropertyValueFactory();
+	ADVANCEDWIDGETS_API FPropertyValueFactory();
 	FPropertyValueFactory(const FPropertyValueFactory&) = delete;
 	FPropertyValueFactory& operator=(const FPropertyValueFactory&) = delete;
 
-	TSharedPtr<SWidget> Generate(FGenerateArgs Args) const;
-	TSharedPtr<SWidget> GenerateDefault(FGenerateArgs Args) const;
+	ADVANCEDWIDGETS_API TSharedPtr<SWidget> Generate(FGenerateArgs Args) const;
+	ADVANCEDWIDGETS_API TSharedPtr<SWidget> GenerateDefault(FGenerateArgs Args) const;
 
-	bool HasCustomPropertyValue(const FFieldClass* Property) const;
-	bool HasCustomPropertyValue(const UStruct* Struct) const;
+	ADVANCEDWIDGETS_API bool HasCustomPropertyValue(const FFieldClass* Property) const;
+	ADVANCEDWIDGETS_API bool HasCustomPropertyValue(const UStruct* Struct) const;
 
-	FHandle Register(const FFieldClass* Field, FOnGenerate OnGenerate);
-	FHandle Register(const UStruct* Struct, FOnGenerate OnGenerate);
-	void Unregister(FHandle Handle);
+	ADVANCEDWIDGETS_API FHandle Register(const FFieldClass* Field, FOnGenerate OnGenerate);
+	ADVANCEDWIDGETS_API FHandle Register(const UStruct* Struct, FOnGenerate OnGenerate);
+	ADVANCEDWIDGETS_API void Unregister(FHandle Handle);
 
 private:
-	static FHandle MakeHandle();
+	static ADVANCEDWIDGETS_API FHandle MakeHandle();
 	TPimplPtr<Private::FPropertyValueFactoryImpl> Impl;
 };
 
diff --git a/Engine/Source/Runtime/AdvancedWidgets/Public/Styling/AdvancedWidgetsStyle.h b/Engine/Source/Runtime/AdvancedWidgets/Public/Styling/AdvancedWidgetsStyle.h
--- a/Engine/Source/Runtime/AdvancedWidgets/Public/Styling/AdvancedWidgetsStyle.h
+++ b/Engine/Source/Runtime/AdvancedWidgets/Public/Styling/AdvancedWidgetsStyle.h
@@ -13,12 +13,12 @@ namespace UE::AdvancedWidgets
 /**
  * Style for advanced widgets
  */
-class ADVANCEDWIDGETS_API FAdvancedWidgetsStyle
+class FAdvancedWidgetsStyle
 {
 public:
-	static void Create();
-	static void Destroy();
-	static const ISlateStyle& Get();
+	static ADVANCEDWIDGETS_API void Create();
+	static ADVANCEDWIDGETS_API void Destroy();
+	static ADVANCEDWIDGETS_API const ISlateStyle& Get();
 
 	static const ::UE::PropertyViewer::FFieldColorSettings& GetColorSettings()
 	{
@@ -27,10 +27,10 @@ public:
 
 private:
 	/** Singleton instances of this style. */
-	static TUniquePtr<FSlateStyleSet> Instance;
+	static ADVANCEDWIDGETS_API TUniquePtr<FSlateStyleSet> Instance;
 
 	/**  */
 	static ::UE::PropertyViewer::FFieldColorSettings ColorSettings;
 };
 
-} // namespace
\ No newline at end of file
+} // namespace
diff --git a/Engine/Source/Runtime/AdvancedWidgets/Public/Widgets/PropertyViewer/SBoolPropertyValue.h b/Engine/Source/Runtime/AdvancedWidgets/Public/Widgets/PropertyViewer/SBoolPropertyValue.h
--- a/Engine/Source/Runtime/AdvancedWidgets/Public/Widgets/PropertyViewer/SBoolPropertyValue.h
+++ b/Engine/Source/Runtime/AdvancedWidgets/Public/Widgets/PropertyViewer/SBoolPropertyValue.h
@@ -14,10 +14,10 @@ namespace UE::PropertyViewer
 class INotifyHook;
 
 /** */
-class ADVANCEDWIDGETS_API SBoolPropertyValue : public SCompoundWidget
+class SBoolPropertyValue : public SCompoundWidget
 {
 public:
-	static TSharedPtr<SWidget> CreateInstance(const FPropertyValueFactory::FGenerateArgs Args);
+	static ADVANCEDWIDGETS_API TSharedPtr<SWidget> CreateInstance(const FPropertyValueFactory::FGenerateArgs Args);
 
 public:
 	SLATE_BEGIN_ARGS(SBoolPropertyValue) {}
@@ -25,7 +25,7 @@ public:
 		SLATE_ARGUMENT(INotifyHook*, NotifyHook);
 	SLATE_END_ARGS()
 
-	void Construct(const FArguments& InArgs);
+	ADVANCEDWIDGETS_API void Construct(const FArguments& InArgs);
 
 private:
 	ECheckBoxState HandleIsChecked() const;
diff --git a/Engine/Source/Runtime/AdvancedWidgets/Public/Widgets/PropertyViewer/SDefaultPropertyValue.h b/Engine/Source/Runtime/AdvancedWidgets/Public/Widgets/PropertyViewer/SDefaultPropertyValue.h
--- a/Engine/Source/Runtime/AdvancedWidgets/Public/Widgets/PropertyViewer/SDefaultPropertyValue.h
+++ b/Engine/Source/Runtime/AdvancedWidgets/Public/Widgets/PropertyViewer/SDefaultPropertyValue.h
@@ -11,17 +11,17 @@ namespace UE::PropertyViewer
 {
 
 /** */
-class ADVANCEDWIDGETS_API SDefaultPropertyValue : public SCompoundWidget
+class SDefaultPropertyValue : public SCompoundWidget
 {
 public:
-	static TSharedPtr<SWidget> CreateInstance(const FPropertyValueFactory::FGenerateArgs Args);
+	static ADVANCEDWIDGETS_API TSharedPtr<SWidget> CreateInstance(const FPropertyValueFactory::FGenerateArgs Args);
 
 public:
 	SLATE_BEGIN_ARGS(SDefaultPropertyValue) {}
 		SLATE_ARGUMENT(FPropertyPath, Path);
 	SLATE_END_ARGS()
 
-	void Construct(const FArguments& InArgs);
+	ADVANCEDWIDGETS_API void Construct(const FArguments& InArgs);
 
 private:
 	FText GetText() const;
diff --git a/Engine/Source/Runtime/AdvancedWidgets/Public/Widgets/PropertyViewer/SEnumPropertyValue.h b/Engine/Source/Runtime/AdvancedWidgets/Public/Widgets/PropertyViewer/SEnumPropertyValue.h
--- a/Engine/Source/Runtime/AdvancedWidgets/Public/Widgets/PropertyViewer/SEnumPropertyValue.h
+++ b/Engine/Source/Runtime/AdvancedWidgets/Public/Widgets/PropertyViewer/SEnumPropertyValue.h
@@ -13,10 +13,10 @@ namespace UE::PropertyViewer
 class INotifyHook;
 
 /** */
-class ADVANCEDWIDGETS_API SEnumPropertyValue : public SCompoundWidget
+class SEnumPropertyValue : public SCompoundWidget
 {
 public:
-	static TSharedPtr<SWidget> CreateInstance(const FPropertyValueFactory::FGenerateArgs Args);
+	static ADVANCEDWIDGETS_API TSharedPtr<SWidget> CreateInstance(const FPropertyValueFactory::FGenerateArgs Args);
 
 public:
 	SLATE_BEGIN_ARGS(SEnumPropertyValue) {}
@@ -24,7 +24,7 @@ public:
 		SLATE_ARGUMENT(INotifyHook*, NotifyHook);
 	SLATE_END_ARGS()
 
-	void Construct(const FArguments& InArgs);
+	ADVANCEDWIDGETS_API void Construct(const FArguments& InArgs);
 
 private:
 	FText GetText() const;
diff --git a/Engine/Source/Runtime/AdvancedWidgets/Public/Widgets/PropertyViewer/SFieldIcon.h b/Engine/Source/Runtime/AdvancedWidgets/Public/Widgets/PropertyViewer/SFieldIcon.h
--- a/Engine/Source/Runtime/AdvancedWidgets/Public/Widgets/PropertyViewer/SFieldIcon.h
+++ b/Engine/Source/Runtime/AdvancedWidgets/Public/Widgets/PropertyViewer/SFieldIcon.h
@@ -10,17 +10,17 @@ namespace UE::PropertyViewer
 {
 
 /** */
-class ADVANCEDWIDGETS_API SFieldIcon : public SCompoundWidget
+class SFieldIcon : public SCompoundWidget
 {
 public:
 	SLATE_BEGIN_ARGS(SFieldIcon) {}
 		SLATE_ARGUMENT(TOptional<const FFieldColorSettings>, OverrideColorSettings);
 	SLATE_END_ARGS()
 
-	void Construct(const FArguments& InArgs, const UClass* Class);
-	void Construct(const FArguments& InArgs, const UScriptStruct* Struct);
-	void Construct(const FArguments& InArgs, const FProperty* Property);
-	void Construct(const FArguments& InArgs, const UFunction* FunctionToDisplay);
+	ADVANCEDWIDGETS_API void Construct(const FArguments& InArgs, const UClass* Class);
+	ADVANCEDWIDGETS_API void Construct(const FArguments& InArgs, const UScriptStruct* Struct);
+	ADVANCEDWIDGETS_API void Construct(const FArguments& InArgs, const FProperty* Property);
+	ADVANCEDWIDGETS_API void Construct(const FArguments& InArgs, const UFunction* FunctionToDisplay);
 };
 
 } //namespace
diff --git a/Engine/Source/Runtime/AdvancedWidgets/Public/Widgets/PropertyViewer/SFieldName.h b/Engine/Source/Runtime/AdvancedWidgets/Public/Widgets/PropertyViewer/SFieldName.h
--- a/Engine/Source/Runtime/AdvancedWidgets/Public/Widgets/PropertyViewer/SFieldName.h
+++ b/Engine/Source/Runtime/AdvancedWidgets/Public/Widgets/PropertyViewer/SFieldName.h
@@ -12,7 +12,7 @@ namespace UE::PropertyViewer
 {
 
 /** */
-class ADVANCEDWIDGETS_API SFieldName : public SCompoundWidget
+class SFieldName : public SCompoundWidget
 {
 public:
 	SLATE_BEGIN_ARGS(SFieldName)
@@ -27,12 +27,12 @@ public:
 		SLATE_ATTRIBUTE(FText, HighlightText)
 	SLATE_END_ARGS()
 
-	void Construct(const FArguments& InArgs, const UClass* Class);
-	void Construct(const FArguments& InArgs, const UScriptStruct* Struct);
-	void Construct(const FArguments& InArgs, const FProperty* Property);
-	void Construct(const FArguments& InArgs, const UFunction* FunctionToDisplay);
+	ADVANCEDWIDGETS_API void Construct(const FArguments& InArgs, const UClass* Class);
+	ADVANCEDWIDGETS_API void Construct(const FArguments& InArgs, const UScriptStruct* Struct);
+	ADVANCEDWIDGETS_API void Construct(const FArguments& InArgs, const FProperty* Property);
+	ADVANCEDWIDGETS_API void Construct(const FArguments& InArgs, const UFunction* FunctionToDisplay);
 
-	void SetHighlightText(TAttribute<FText> InHighlightText);
+	ADVANCEDWIDGETS_API void SetHighlightText(TAttribute<FText> InHighlightText);
 
 private:
 	FText GetToolTipText() const;
diff --git a/Engine/Source/Runtime/AdvancedWidgets/Public/Widgets/PropertyViewer/SNumericPropertyValue.h b/Engine/Source/Runtime/AdvancedWidgets/Public/Widgets/PropertyViewer/SNumericPropertyValue.h
--- a/Engine/Source/Runtime/AdvancedWidgets/Public/Widgets/PropertyViewer/SNumericPropertyValue.h
+++ b/Engine/Source/Runtime/AdvancedWidgets/Public/Widgets/PropertyViewer/SNumericPropertyValue.h
@@ -13,10 +13,10 @@ namespace UE::PropertyViewer
 class INotifyHook;
 
 /** */
-class ADVANCEDWIDGETS_API SNumericPropertyValue : public SCompoundWidget
+class SNumericPropertyValue : public SCompoundWidget
 {
 public:
-	static TSharedPtr<SWidget> CreateInstance(const FPropertyValueFactory::FGenerateArgs Args);
+	static ADVANCEDWIDGETS_API TSharedPtr<SWidget> CreateInstance(const FPropertyValueFactory::FGenerateArgs Args);
 
 public:
 	SLATE_BEGIN_ARGS(SNumericPropertyValue) {}
@@ -24,7 +24,7 @@ public:
 		SLATE_ARGUMENT(INotifyHook*, NotifyHook);
 	SLATE_END_ARGS()
 
-	void Construct(const FArguments& InArgs);
+	ADVANCEDWIDGETS_API void Construct(const FArguments& InArgs);
 
 private:
 	void OnEndSliderMovement_uint64(uint64 Value);
diff --git a/Engine/Source/Runtime/AdvancedWidgets/Public/Widgets/PropertyViewer/SPropertyViewer.h b/Engine/Source/Runtime/AdvancedWidgets/Public/Widgets/PropertyViewer/SPropertyViewer.h
--- a/Engine/Source/Runtime/AdvancedWidgets/Public/Widgets/PropertyViewer/SPropertyViewer.h
+++ b/Engine/Source/Runtime/AdvancedWidgets/Public/Widgets/PropertyViewer/SPropertyViewer.h
@@ -22,7 +22,7 @@ class FPropertyViewerImpl;
 }
 
 /** */
-class ADVANCEDWIDGETS_API SPropertyViewer : public SCompoundWidget
+class SPropertyViewer : public SCompoundWidget
 {
 private:
 	using Super = SCompoundWidget;
@@ -114,29 +114,29 @@ public:
 		SLATE_EVENT(FOnDragDetected, OnDragDetected);
 	SLATE_END_ARGS()
 
-	void Construct(const FArguments& InArgs);
-	void Construct(const FArguments& InArgs, const UScriptStruct* Struct);
-	void Construct(const FArguments& InArgs, const UScriptStruct* Struct, void* Data);
-	void Construct(const FArguments& InArgs, const UClass* Class);
-	void Construct(const FArguments& InArgs, UObject* ObjectInstance);
-	void Construct(const FArguments& InArgs, const UFunction* Function);
+	ADVANCEDWIDGETS_API void Construct(const FArguments& InArgs);
+	ADVANCEDWIDGETS_API void Construct(const FArguments& InArgs, const UScriptStruct* Struct);
+	ADVANCEDWIDGETS_API void Construct(const FArguments& InArgs, const UScriptStruct* Struct, void* Data);
+	ADVANCEDWIDGETS_API void Construct(const FArguments& InArgs, const UClass* Class);
+	ADVANCEDWIDGETS_API void Construct(const FArguments& InArgs, UObject* ObjectInstance);
+	ADVANCEDWIDGETS_API void Construct(const FArguments& InArgs, const UFunction* Function);
 
-	virtual void Tick(const FGeometry& AllottedGeometry, const double InCurrentTime, const float InDeltaTime) override;
+	ADVANCEDWIDGETS_API virtual void Tick(const FGeometry& AllottedGeometry, const double InCurrentTime, const float InDeltaTime) override;
 
 public:
-	FHandle AddContainer(const UScriptStruct* Struct, TOptional<FText> DisplayName = TOptional<FText>());
-	FHandle AddContainer(const UClass* Class, TOptional<FText> DisplayName = TOptional<FText>());
-	FHandle AddContainer(const UFunction* Function, TOptional<FText> DisplayName = TOptional<FText>());
-	FHandle AddInstance(const UScriptStruct* Struct, void* Data, TOptional<FText> DisplayName = TOptional<FText>());
-	FHandle AddInstance(UObject* ObjectInstance, TOptional<FText> DisplayName = TOptional<FText>());
+	ADVANCEDWIDGETS_API FHandle AddContainer(const UScriptStruct* Struct, TOptional<FText> DisplayName = TOptional<FText>());
+	ADVANCEDWIDGETS_API FHandle AddContainer(const UClass* Class, TOptional<FText> DisplayName = TOptional<FText>());
+	ADVANCEDWIDGETS_API FHandle AddContainer(const UFunction* Function, TOptional<FText> DisplayName = TOptional<FText>());
+	ADVANCEDWIDGETS_API FHandle AddInstance(const UScriptStruct* Struct, void* Data, TOptional<FText> DisplayName = TOptional<FText>());
+	ADVANCEDWIDGETS_API FHandle AddInstance(UObject* ObjectInstance, TOptional<FText> DisplayName = TOptional<FText>());
 
-	void Remove(FHandle Identifier);
-	void RemoveAll();
+	ADVANCEDWIDGETS_API void Remove(FHandle Identifier);
+	ADVANCEDWIDGETS_API void RemoveAll();
 
-	TArray<FSelectedItem> GetSelectedItems() const;
+	ADVANCEDWIDGETS_API TArray<FSelectedItem> GetSelectedItems() const;
 
-	void SetRawFilterText(const FText& InFilterText);
-	void SetSelection(FHandle Container, TArrayView<const FFieldVariant> FieldPath);
+	ADVANCEDWIDGETS_API void SetRawFilterText(const FText& InFilterText);
+	ADVANCEDWIDGETS_API void SetSelection(FHandle Container, TArrayView<const FFieldVariant> FieldPath);
 
 private:
 	void ConstructInternal(const FArguments& InArgs);
diff --git a/Engine/Source/Runtime/AdvancedWidgets/Public/Widgets/PropertyViewer/SStringPropertyValue.h b/Engine/Source/Runtime/AdvancedWidgets/Public/Widgets/PropertyViewer/SStringPropertyValue.h
--- a/Engine/Source/Runtime/AdvancedWidgets/Public/Widgets/PropertyViewer/SStringPropertyValue.h
+++ b/Engine/Source/Runtime/AdvancedWidgets/Public/Widgets/PropertyViewer/SStringPropertyValue.h
@@ -13,10 +13,10 @@ namespace UE::PropertyViewer
 class INotifyHook;
 
 /** */
-class ADVANCEDWIDGETS_API SStringPropertyValue : public SCompoundWidget
+class SStringPropertyValue : public SCompoundWidget
 {
 public:
-	static TSharedPtr<SWidget> CreateInstance(const FPropertyValueFactory::FGenerateArgs Args);
+	static ADVANCEDWIDGETS_API TSharedPtr<SWidget> CreateInstance(const FPropertyValueFactory::FGenerateArgs Args);
 
 public:
 	SLATE_BEGIN_ARGS(SStringPropertyValue) {}
@@ -24,7 +24,7 @@ public:
 		SLATE_ARGUMENT(INotifyHook*, NotifyHook);
 	SLATE_END_ARGS()
 
-	void Construct(const FArguments& InArgs);
+	ADVANCEDWIDGETS_API void Construct(const FArguments& InArgs);
 
 private:
 	FText GetText() const;
diff --git a/Engine/Source/Runtime/AdvancedWidgets/Public/Widgets/SRadialSlider.h b/Engine/Source/Runtime/AdvancedWidgets/Public/Widgets/SRadialSlider.h
--- a/Engine/Source/Runtime/AdvancedWidgets/Public/Widgets/SRadialSlider.h
+++ b/Engine/Source/Runtime/AdvancedWidgets/Public/Widgets/SRadialSlider.h
@@ -22,7 +22,7 @@ class FPaintArgs;
 class FSlateWindowElementList;
 class UCurveFloat;
 
-class ADVANCEDWIDGETS_API SRadialSlider
+class SRadialSlider
 	: public SLeafWidget
 {
 public:
@@ -144,23 +144,23 @@ public:
 
 	SLATE_END_ARGS()
 
-	SRadialSlider();
+	ADVANCEDWIDGETS_API SRadialSlider();
 
 	/**
 	 * Construct the widget.
 	 * 
 	 * @param InDeclaration A declaration from which to construct the widget.
 	 */
-	void Construct( const SRadialSlider::FArguments& InDeclaration );
+	ADVANCEDWIDGETS_API void Construct( const SRadialSlider::FArguments& InDeclaration );
 
 	/** Get the SliderRange attribute */
 	FRuntimeFloatCurve GetSliderRange() const { return SliderRange; }
 
 	/** Get the minumum value in Slider Range */
-	float GetMinValue() const;
+	ADVANCEDWIDGETS_API float GetMinValue() const;
 
 	/** Get the maximum value in Slider Range */
-	float GetMaxValue() const;
+	ADVANCEDWIDGETS_API float GetMaxValue() const;
 
 	/** Get the MinSliderHandleAngle attribute */
 	float GetSliderHandleStartAngle() const { return SliderHandleStartAngle; }
@@ -175,108 +175,108 @@ public:
 	TArray<float> GetValueTags() const { return ValueTags;	}
 
 	/** Get the Value attribute */
-	float GetValue() const;
+	ADVANCEDWIDGETS_API float GetValue() const;
 
 	/** Get the bUseCustomDefaultValue attribute */
-	bool GetUseCustomDefaultValue() const;
+	ADVANCEDWIDGETS_API bool GetUseCustomDefaultValue() const;
 
 	/** Get ths CustomDefaultValue attribute */
-	float GetCustomDefaultValue() const;
+	ADVANCEDWIDGETS_API float GetCustomDefaultValue() const;
 
 	/** Get the Value attribute scaled from 0 to 1 */
-	float GetNormalizedValue(float RawValue) const;
+	ADVANCEDWIDGETS_API float GetNormalizedValue(float RawValue) const;
 
 	/** Get the Slider's Handle position scaled from 0 to 1 */
-	float GetNormalizedSliderHandlePosition() const;
+	ADVANCEDWIDGETS_API float GetNormalizedSliderHandlePosition() const;
 
 	/** Set the Value attribute */
-	void SetValue(const TAttribute<float>& InValueAttribute);
+	ADVANCEDWIDGETS_API void SetValue(const TAttribute<float>& InValueAttribute);
 
 	/** Set the bUseCustomDefaultValue attribute */
-	void SetUseCustomDefaultValue(const TAttribute<bool>& InValueAttribute);
+	ADVANCEDWIDGETS_API void SetUseCustomDefaultValue(const TAttribute<bool>& InValueAttribute);
 
 	/** Set the CustomDefaultValue attribute */
-	void SetCustomDefaultValue(const TAttribute<float>& InValueAttribute);
+	ADVANCEDWIDGETS_API void SetCustomDefaultValue(const TAttribute<float>& InValueAttribute);
 	
 	/** Set the SliderRange attribute */
 	void SetSliderRange(const FRuntimeFloatCurve& InSliderRange) { SliderRange = InSliderRange; }
 
 	/** Set the SliderHandleStartAngle and SliderHandleEndAngle attributes. If the new SliderHandleStartAngle is more than the new SliderHandleEndAngle, SliderHandleEndAngle will be changed to equal SliderHandleStartAngle. */
-	void SetSliderHandleStartAngleAndSliderHandleEndAngle(float InSliderHandleStartAngle, float InSliderHandleEndAngle);
+	ADVANCEDWIDGETS_API void SetSliderHandleStartAngleAndSliderHandleEndAngle(float InSliderHandleStartAngle, float InSliderHandleEndAngle);
 	
 	/** Set the AngularOffset attribute */
 	void SetAngularOffset(float InAngularOffset) { AngularOffset = InAngularOffset; }
 
 	/** Set the HandStartEndRatio. Clamped to 0.0 to 1.0, and if the start ratio is more than the end ratio, end ratio will be set to the start ratio.  */
-	void SetHandStartEndRatio(FVector2D InHandStartEndRatio);
+	ADVANCEDWIDGETS_API void SetHandStartEndRatio(FVector2D InHandStartEndRatio);
 
 	/** Set the ValueTags attribute */
 	void SetValueTags(const TArray<float>& InValueTags) { ValueTags = InValueTags; }
 
 	/** Set the Locked attribute */
-	void SetLocked(const TAttribute<bool>& InLocked);
+	ADVANCEDWIDGETS_API void SetLocked(const TAttribute<bool>& InLocked);
 	
 	/** Set the SliderBarColor attribute */
-	void SetSliderBarColor(FSlateColor InSliderBarColor);
+	ADVANCEDWIDGETS_API void SetSliderBarColor(FSlateColor InSliderBarColor);
 	
 	/** Set the SliderProgressColor attribute */
-	void SetSliderProgressColor(FSlateColor InSliderProgressColor);
+	ADVANCEDWIDGETS_API void SetSliderProgressColor(FSlateColor InSliderProgressColor);
 
 	/** Set the SliderHandleColor attribute */
-	void SetSliderHandleColor(FSlateColor InSliderHandleColor);
+	ADVANCEDWIDGETS_API void SetSliderHandleColor(FSlateColor InSliderHandleColor);
 
 	/** Set the SliderHandleColor attribute */
-	void SetCenterBackgroundColor(FSlateColor InCenterHandleColor);
+	ADVANCEDWIDGETS_API void SetCenterBackgroundColor(FSlateColor InCenterHandleColor);
 
 	/** Set the Thickness attribute. For backward compatibility, Thickness will be used for drawing instead of Style->BarThickness only if it has been set with this method,  */
-	void SetThickness(const float InThickness);
+	ADVANCEDWIDGETS_API void SetThickness(const float InThickness);
 
 	/** Get the StepSize attribute */
-	float GetStepSize() const;
+	ADVANCEDWIDGETS_API float GetStepSize() const;
 
 	/** Set the StepSize attribute */
-	void SetStepSize(const TAttribute<float>& InStepSize);
+	ADVANCEDWIDGETS_API void SetStepSize(const TAttribute<float>& InStepSize);
 
 	/** Set the MouseUsesStep attribute */
-	void SetMouseUsesStep(bool MouseUsesStep);
+	ADVANCEDWIDGETS_API void SetMouseUsesStep(bool MouseUsesStep);
 
 	/** Set the RequiresControllerLock attribute */
-	void SetRequiresControllerLock(bool RequiresControllerLock);
+	ADVANCEDWIDGETS_API void SetRequiresControllerLock(bool RequiresControllerLock);
 
 	/** Set the UseVerticalDrag attribute */
-	void SetUseVerticalDrag(bool UseVerticalDrag);
+	ADVANCEDWIDGETS_API void SetUseVerticalDrag(bool UseVerticalDrag);
 
 	/** Set the ShowSliderHandle attribute */
-	void SetShowSliderHandle(bool ShowSliderHandle);
+	ADVANCEDWIDGETS_API void SetShowSliderHandle(bool ShowSliderHandle);
 
 	/** Set the ShowSliderHand attribute */
-	void SetShowSliderHand(bool ShowSliderHand);
+	ADVANCEDWIDGETS_API void SetShowSliderHand(bool ShowSliderHand);
 
 public:
 
 	// SWidget overrides
 
-	virtual int32 OnPaint( const FPaintArgs& Args, const FGeometry& AllottedGeometry, const FSlateRect& MyCullingRect, FSlateWindowElementList& OutDrawElements, int32 LayerId, const FWidgetStyle& InWidgetStyle, bool bParentEnabled ) const override;
-	virtual FVector2D ComputeDesiredSize(float) const override;
-	virtual FReply OnMouseButtonDown( const FGeometry& MyGeometry, const FPointerEvent& MouseEvent ) override;
-	virtual FReply OnMouseButtonUp( const FGeometry& MyGeometry, const FPointerEvent& MouseEvent ) override;
-	virtual FReply OnMouseMove( const FGeometry& MyGeometry, const FPointerEvent& MouseEvent ) override;
-	virtual FReply OnTouchStarted(const FGeometry& MyGeometry, const FPointerEvent& InTouchEvent) override;
-	virtual FReply OnTouchMoved(const FGeometry& MyGeometry, const FPointerEvent& InTouchEvent) override;
-	virtual FReply OnTouchEnded(const FGeometry& MyGeometry, const FPointerEvent& InTouchEvent) override;
-	virtual FReply OnKeyDown(const FGeometry& MyGeometry, const FKeyEvent& InKeyEvent) override;
-	virtual FReply OnKeyUp(const FGeometry& MyGeometry, const FKeyEvent& InKeyEvent) override;
-	virtual void OnFocusLost(const FFocusEvent& InFocusEvent) override;
-	virtual FNavigationReply OnNavigation(const FGeometry& MyGeometry, const FNavigationEvent& InNavigationEvent) override;
-
-	virtual bool SupportsKeyboardFocus() const override;
-	virtual bool IsInteractable() const override;
+	ADVANCEDWIDGETS_API virtual int32 OnPaint( const FPaintArgs& Args, const FGeometry& AllottedGeometry, const FSlateRect& MyCullingRect, FSlateWindowElementList& OutDrawElements, int32 LayerId, const FWidgetStyle& InWidgetStyle, bool bParentEnabled ) const override;
+	ADVANCEDWIDGETS_API virtual FVector2D ComputeDesiredSize(float) const override;
+	ADVANCEDWIDGETS_API virtual FReply OnMouseButtonDown( const FGeometry& MyGeometry, const FPointerEvent& MouseEvent ) override;
+	ADVANCEDWIDGETS_API virtual FReply OnMouseButtonUp( const FGeometry& MyGeometry, const FPointerEvent& MouseEvent ) override;
+	ADVANCEDWIDGETS_API virtual FReply OnMouseMove( const FGeometry& MyGeometry, const FPointerEvent& MouseEvent ) override;
+	ADVANCEDWIDGETS_API virtual FReply OnTouchStarted(const FGeometry& MyGeometry, const FPointerEvent& InTouchEvent) override;
+	ADVANCEDWIDGETS_API virtual FReply OnTouchMoved(const FGeometry& MyGeometry, const FPointerEvent& InTouchEvent) override;
+	ADVANCEDWIDGETS_API virtual FReply OnTouchEnded(const FGeometry& MyGeometry, const FPointerEvent& InTouchEvent) override;
+	ADVANCEDWIDGETS_API virtual FReply OnKeyDown(const FGeometry& MyGeometry, const FKeyEvent& InKeyEvent) override;
+	ADVANCEDWIDGETS_API virtual FReply OnKeyUp(const FGeometry& MyGeometry, const FKeyEvent& InKeyEvent) override;
+	ADVANCEDWIDGETS_API virtual void OnFocusLost(const FFocusEvent& InFocusEvent) override;
+	ADVANCEDWIDGETS_API virtual FNavigationReply OnNavigation(const FGeometry& MyGeometry, const FNavigationEvent& InNavigationEvent) override;
+
+	ADVANCEDWIDGETS_API virtual bool SupportsKeyboardFocus() const override;
+	ADVANCEDWIDGETS_API virtual bool IsInteractable() const override;
 #if WITH_ACCESSIBILITY
-	virtual TSharedRef<FSlateAccessibleWidget> CreateAccessibleWidget() override;
+	ADVANCEDWIDGETS_API virtual TSharedRef<FSlateAccessibleWidget> CreateAccessibleWidget() override;
 #endif
 
 	/** @return Is the handle locked or not? Defaults to false */
-	bool IsLocked() const;
+	ADVANCEDWIDGETS_API bool IsLocked() const;
 
 protected:
 
@@ -285,7 +285,7 @@ protected:
 	 *
 	 * @param NewValue The value to commit.
 	 */
-	virtual void CommitValue(float NewValue);
+	ADVANCEDWIDGETS_API virtual void CommitValue(float NewValue);
 
 	/**
 	 * Calculates the new value based on the given absolute coordinates.
@@ -294,10 +294,10 @@ protected:
 	 * @param AbsolutePosition The absolute position of the slider.
 	 * @return The new value.
 	 */
-	float PositionToValue(const FGeometry& MyGeometry, const FVector2D& AbsolutePosition);
+	ADVANCEDWIDGETS_API float PositionToValue(const FGeometry& MyGeometry, const FVector2D& AbsolutePosition);
 	
-	const FSlateBrush* GetBarImage() const;
-	const FSlateBrush* GetThumbImage() const;
+	ADVANCEDWIDGETS_API const FSlateBrush* GetBarImage() const;
+	ADVANCEDWIDGETS_API const FSlateBrush* GetThumbImage() const;
 
 protected:
 
@@ -386,16 +386,16 @@ protected:
 private:
 
 	// Resets controller input state. Fires delegates.
-	void ResetControllerState();
+	ADVANCEDWIDGETS_API void ResetControllerState();
 
 	// Helper function to convert Absolute Position to Angle
-	float GetAngleFromPosition(const FGeometry& MyGeometry, const FVector2D& AbsolutePosition);
+	ADVANCEDWIDGETS_API float GetAngleFromPosition(const FGeometry& MyGeometry, const FVector2D& AbsolutePosition);
 	
 	// Called on Mouse / Touch input to cache relevant properties
-	void OnInputStarted(const FGeometry& MyGeometry, const FVector2D& InputAbsolutePosition);
+	ADVANCEDWIDGETS_API void OnInputStarted(const FGeometry& MyGeometry, const FVector2D& InputAbsolutePosition);
 
 	// Helper function for adding slider points to a slider points array
-	void AddSliderPointToArray(TArray<FVector2D>& SliderPoints, const bool bIsUnique, const FVector2D& SliderPoint) const;
+	ADVANCEDWIDGETS_API void AddSliderPointToArray(TArray<FVector2D>& SliderPoints, const bool bIsUnique, const FVector2D& SliderPoint) const;
 
 	// Holds a delegate that is executed when the mouse is pressed and a capture begins.
 	FSimpleDelegate OnMouseCaptureBegin;
diff --git a/Engine/Source/Runtime/Analytics/Analytics/Public/AnalyticsPropertyStore.h b/Engine/Source/Runtime/Analytics/Analytics/Public/AnalyticsPropertyStore.h
--- a/Engine/Source/Runtime/Analytics/Analytics/Public/AnalyticsPropertyStore.h
+++ b/Engine/Source/Runtime/Analytics/Analytics/Public/AnalyticsPropertyStore.h
@@ -32,12 +32,12 @@ struct FAnalyticsEventAttribute;
  * The string time is often updated in place if enough capacity was reserved. All store operations are
  * thread safe and atomic from the caller perspective.
  */
-class ANALYTICS_API FAnalyticsPropertyStore : public IAnalyticsPropertyStore
+class FAnalyticsPropertyStore : public IAnalyticsPropertyStore
 {
 public:
 	/** Create an empty store. */
-	FAnalyticsPropertyStore();
-	virtual ~FAnalyticsPropertyStore();
+	ANALYTICS_API FAnalyticsPropertyStore();
+	ANALYTICS_API virtual ~FAnalyticsPropertyStore();
 
 	FAnalyticsPropertyStore(const FAnalyticsPropertyStore&) = delete;
 	FAnalyticsPropertyStore(FAnalyticsPropertyStore&&) = delete;
@@ -50,60 +50,60 @@ public:
 	 * @param CapacityHint The desired file capacity. The implementation has a minimum and maximum capacity to reserve.
 	 * @param true if the file was created and the store created, false otherwise.
 	 */
-	bool Create(const FString& Pathname, uint32 CapacityHint = 4 * 1024);
+	ANALYTICS_API bool Create(const FString& Pathname, uint32 CapacityHint = 4 * 1024);
 
 	/**
 	 * Load the specified file in memory, resetting the current state.
 	 * @param Pathname The path to the file to open.
 	 * @return true if the file was loaded successuflly, false if the file doesn't exist, the file format/version is invalid or if the file is corrupted (failed checksum).
 	 */
-	bool Load(const FString& Pathname);
+	ANALYTICS_API bool Load(const FString& Pathname);
 
 	/** Returns true is the store was successfully created or loaded. */
 	bool IsValid() const { return FileHandle != nullptr; }
 
 	//~ IAnalyticsPropertyStore interface
-	virtual uint32 Num() const override;
-	virtual bool Contains(const FString& Key) const override;
-	virtual bool Remove(const FString& Key) override;
-	virtual void RemoveAll() override;
-	virtual EStatusCode Set(const FString& Key, int32 Value) override;
-	virtual EStatusCode Set(const FString& Key, uint32 Value) override;
-	virtual EStatusCode Set(const FString& Key, int64 Value) override;
-	virtual EStatusCode Set(const FString& Key, uint64 Value) override;
-	virtual EStatusCode Set(const FString& Key, float Value) override;
-	virtual EStatusCode Set(const FString& Key, double Value) override;
-	virtual EStatusCode Set(const FString& Key, bool Value) override;
-	virtual EStatusCode Set(const FString& Key, const FString& Value, uint32 CharCountCapacityHint = 0) override;
-	virtual EStatusCode Set(const FString& Key, const FDateTime& Value) override;
-	virtual EStatusCode Set(const FString& Key, int32 Value, const TFunction<bool(const int32* /*Actual*/, const int32& /*Proposed*/)>& ConditionFn) override;
-	virtual EStatusCode Set(const FString& Key, uint32 Value, const TFunction<bool(const uint32* /*Actual*/, const uint32& /*Proposed*/)>& ConditionFn) override;
-	virtual EStatusCode Set(const FString& Key, int64 Value, const TFunction<bool(const int64* /*Actual*/, const int64& /*Proposed*/)>& ConditionFn) override;
-	virtual EStatusCode Set(const FString& Key, uint64 Value, const TFunction<bool(const uint64* /*Actual*/, const uint64& /*Proposed*/)>& ConditionFn) override;
-	virtual EStatusCode Set(const FString& Key, float Value, const TFunction<bool(const float* /*Actual*/, const float& /*Proposed*/)>& ConditionFn) override;
-	virtual EStatusCode Set(const FString& Key, double Value, const TFunction<bool(const double* /*Actual*/, const double& /*Proposed*/)>& ConditionFn) override;
-	virtual EStatusCode Set(const FString& Key, bool Value, const TFunction<bool(const bool* /*Actual*/, const bool& /*Proposed*/)>& ConditionFn) override;
-	virtual EStatusCode Set(const FString& Key, const FString& Value, const TFunction<bool(const FString* /*Actual*/, const FString& /*Proposed*/)>& ConditionFn) override;
-	virtual EStatusCode Set(const FString& Key, const FDateTime& Value, const TFunction<bool(const FDateTime* /*Actual*/, const FDateTime& /*Proposed*/)>& ConditionFn) override;
-	virtual EStatusCode Update(const FString& Key, const TFunction<bool(int32&)>& UpdateFn) override;
-	virtual EStatusCode Update(const FString& Key, const TFunction<bool(uint32&)>& UpdateFn) override;
-	virtual EStatusCode Update(const FString& Key, const TFunction<bool(int64&)>& UpdateFn) override;
-	virtual EStatusCode Update(const FString& Key, const TFunction<bool(uint64&)>& UpdateFn) override;
-	virtual EStatusCode Update(const FString& Key, const TFunction<bool(float&)>& UpdateFn) override;
-	virtual EStatusCode Update(const FString& Key, const TFunction<bool(double&)>& UpdateFn) override;
-	virtual EStatusCode Update(const FString& Key, const TFunction<bool(bool&)>& UpdateFn) override;
-	virtual EStatusCode Update(const FString& Key, const TFunction<bool(FString&)>& UpdateFn) override;
-	virtual EStatusCode Update(const FString& Key, const TFunction<bool(FDateTime&)>& UpdateFn) override;
-	virtual EStatusCode Get(const FString& Key, int32& OutValue) const override;
-	virtual EStatusCode Get(const FString& Key, uint32& OutValue) const override;
-	virtual EStatusCode Get(const FString& Key, int64& OutValue) const override;
-	virtual EStatusCode Get(const FString& Key, uint64& OutValue) const override;
-	virtual EStatusCode Get(const FString& Key, float& OutValue) const override;
-	virtual EStatusCode Get(const FString& Key, double& OutValue) const override;
-	virtual EStatusCode Get(const FString& Key, bool& OutValue) const override;
-	virtual EStatusCode Get(const FString& Key, FString& OutValue) const override;
-	virtual EStatusCode Get(const FString& Key, FDateTime& OutValue) const override;
-	virtual void VisitAll(const TFunction<void(FAnalyticsEventAttribute&&)>& VisitFn) const override;
+	ANALYTICS_API virtual uint32 Num() const override;
+	ANALYTICS_API virtual bool Contains(const FString& Key) const override;
+	ANALYTICS_API virtual bool Remove(const FString& Key) override;
+	ANALYTICS_API virtual void RemoveAll() override;
+	ANALYTICS_API virtual EStatusCode Set(const FString& Key, int32 Value) override;
+	ANALYTICS_API virtual EStatusCode Set(const FString& Key, uint32 Value) override;
+	ANALYTICS_API virtual EStatusCode Set(const FString& Key, int64 Value) override;
+	ANALYTICS_API virtual EStatusCode Set(const FString& Key, uint64 Value) override;
+	ANALYTICS_API virtual EStatusCode Set(const FString& Key, float Value) override;
+	ANALYTICS_API virtual EStatusCode Set(const FString& Key, double Value) override;
+	ANALYTICS_API virtual EStatusCode Set(const FString& Key, bool Value) override;
+	ANALYTICS_API virtual EStatusCode Set(const FString& Key, const FString& Value, uint32 CharCountCapacityHint = 0) override;
+	ANALYTICS_API virtual EStatusCode Set(const FString& Key, const FDateTime& Value) override;
+	ANALYTICS_API virtual EStatusCode Set(const FString& Key, int32 Value, const TFunction<bool(const int32* /*Actual*/, const int32& /*Proposed*/)>& ConditionFn) override;
+	ANALYTICS_API virtual EStatusCode Set(const FString& Key, uint32 Value, const TFunction<bool(const uint32* /*Actual*/, const uint32& /*Proposed*/)>& ConditionFn) override;
+	ANALYTICS_API virtual EStatusCode Set(const FString& Key, int64 Value, const TFunction<bool(const int64* /*Actual*/, const int64& /*Proposed*/)>& ConditionFn) override;
+	ANALYTICS_API virtual EStatusCode Set(const FString& Key, uint64 Value, const TFunction<bool(const uint64* /*Actual*/, const uint64& /*Proposed*/)>& ConditionFn) override;
+	ANALYTICS_API virtual EStatusCode Set(const FString& Key, float Value, const TFunction<bool(const float* /*Actual*/, const float& /*Proposed*/)>& ConditionFn) override;
+	ANALYTICS_API virtual EStatusCode Set(const FString& Key, double Value, const TFunction<bool(const double* /*Actual*/, const double& /*Proposed*/)>& ConditionFn) override;
+	ANALYTICS_API virtual EStatusCode Set(const FString& Key, bool Value, const TFunction<bool(const bool* /*Actual*/, const bool& /*Proposed*/)>& ConditionFn) override;
+	ANALYTICS_API virtual EStatusCode Set(const FString& Key, const FString& Value, const TFunction<bool(const FString* /*Actual*/, const FString& /*Proposed*/)>& ConditionFn) override;
+	ANALYTICS_API virtual EStatusCode Set(const FString& Key, const FDateTime& Value, const TFunction<bool(const FDateTime* /*Actual*/, const FDateTime& /*Proposed*/)>& ConditionFn) override;
+	ANALYTICS_API virtual EStatusCode Update(const FString& Key, const TFunction<bool(int32&)>& UpdateFn) override;
+	ANALYTICS_API virtual EStatusCode Update(const FString& Key, const TFunction<bool(uint32&)>& UpdateFn) override;
+	ANALYTICS_API virtual EStatusCode Update(const FString& Key, const TFunction<bool(int64&)>& UpdateFn) override;
+	ANALYTICS_API virtual EStatusCode Update(const FString& Key, const TFunction<bool(uint64&)>& UpdateFn) override;
+	ANALYTICS_API virtual EStatusCode Update(const FString& Key, const TFunction<bool(float&)>& UpdateFn) override;
+	ANALYTICS_API virtual EStatusCode Update(const FString& Key, const TFunction<bool(double&)>& UpdateFn) override;
+	ANALYTICS_API virtual EStatusCode Update(const FString& Key, const TFunction<bool(bool&)>& UpdateFn) override;
+	ANALYTICS_API virtual EStatusCode Update(const FString& Key, const TFunction<bool(FString&)>& UpdateFn) override;
+	ANALYTICS_API virtual EStatusCode Update(const FString& Key, const TFunction<bool(FDateTime&)>& UpdateFn) override;
+	ANALYTICS_API virtual EStatusCode Get(const FString& Key, int32& OutValue) const override;
+	ANALYTICS_API virtual EStatusCode Get(const FString& Key, uint32& OutValue) const override;
+	ANALYTICS_API virtual EStatusCode Get(const FString& Key, int64& OutValue) const override;
+	ANALYTICS_API virtual EStatusCode Get(const FString& Key, uint64& OutValue) const override;
+	ANALYTICS_API virtual EStatusCode Get(const FString& Key, float& OutValue) const override;
+	ANALYTICS_API virtual EStatusCode Get(const FString& Key, double& OutValue) const override;
+	ANALYTICS_API virtual EStatusCode Get(const FString& Key, bool& OutValue) const override;
+	ANALYTICS_API virtual EStatusCode Get(const FString& Key, FString& OutValue) const override;
+	ANALYTICS_API virtual EStatusCode Get(const FString& Key, FDateTime& OutValue) const override;
+	ANALYTICS_API virtual void VisitAll(const TFunction<void(FAnalyticsEventAttribute&&)>& VisitFn) const override;
 
 	/**
 	 * Flushes cached values to disk. Only one thread can flush at a time, causing possible latency if a flush is alreary happenning.
@@ -118,7 +118,7 @@ public:
 	 *   Store.Flush(true);  // Flush in a background thread, might block if a thread is already flushing. (It doesn't queue several tasks).
 	 * @endcode
 	 */
-	virtual bool Flush(bool bAsync = false, const FTimespan& Timeout = FTimespan::MaxValue()) override;
+	ANALYTICS_API virtual bool Flush(bool bAsync = false, const FTimespan& Timeout = FTimespan::MaxValue()) override;
 
 private:
 	enum class ETypeCode : uint8
@@ -152,18 +152,18 @@ private:
 	friend class FFlushWorker;
 
 private:
-	EStatusCode SetFixedSizeValueInternal(const FString& Key, ETypeCode TypeCode, const uint8* Value, uint32 Size);
-	EStatusCode GetFixedSizeValueInternal(const FString& Key, ETypeCode TypeCode, uint8* OutValue, uint32 Size) const;
-	EStatusCode SetStringValueInternal(const FString& Key, const FString& Value, uint32 CharCountCapacityHint = 0);
-	EStatusCode GetStringValueInternal(const FString& Key, FString& OutValue) const;
-	EStatusCode SetDateTimeValueInternal(const FString& Key, const FDateTime& Value);
-	EStatusCode GetDateTimeValueInternal(const FString& Key, FDateTime& OutValue) const;
-	void Defragment();
+	ANALYTICS_API EStatusCode SetFixedSizeValueInternal(const FString& Key, ETypeCode TypeCode, const uint8* Value, uint32 Size);
+	ANALYTICS_API EStatusCode GetFixedSizeValueInternal(const FString& Key, ETypeCode TypeCode, uint8* OutValue, uint32 Size) const;
+	ANALYTICS_API EStatusCode SetStringValueInternal(const FString& Key, const FString& Value, uint32 CharCountCapacityHint = 0);
+	ANALYTICS_API EStatusCode GetStringValueInternal(const FString& Key, FString& OutValue) const;
+	ANALYTICS_API EStatusCode SetDateTimeValueInternal(const FString& Key, const FDateTime& Value);
+	ANALYTICS_API EStatusCode GetDateTimeValueInternal(const FString& Key, FDateTime& OutValue) const;
+	ANALYTICS_API void Defragment();
 	bool IsFixedSize(ETypeCode Type) const { return Type != ETypeCode::Str; }
 	bool IsDead(ETypeCode Type) const { return EnumHasAnyFlags(Type, ETypeCode::Dead); }
 	ETypeCode RawType(ETypeCode Type) const { return Type & ETypeCode::RawMask; }
-	void FlushInternal(const TArray<uint8>& Data);
-	void Reset();
+	ANALYTICS_API void FlushInternal(const TArray<uint8>& Data);
+	ANALYTICS_API void Reset();
 
 	template<typename T>
 	IAnalyticsPropertyStore::EStatusCode SetFixedSizeValueInternal(const FString& Key, ETypeCode TypeCode, T Value)
diff --git a/Engine/Source/Runtime/Analytics/Analytics/Public/Interfaces/IAnalyticsPropertyStore.h b/Engine/Source/Runtime/Analytics/Analytics/Public/Interfaces/IAnalyticsPropertyStore.h
--- a/Engine/Source/Runtime/Analytics/Analytics/Public/Interfaces/IAnalyticsPropertyStore.h
+++ b/Engine/Source/Runtime/Analytics/Analytics/Public/Interfaces/IAnalyticsPropertyStore.h
@@ -10,7 +10,7 @@
 /**
  * Implements a fast type safe key/value database to store analytics properties collected during a session.
  */
-class ANALYTICS_API IAnalyticsPropertyStore
+class IAnalyticsPropertyStore
 {
 public:
 	/** The list of store operation outcome. */
diff --git a/Engine/Source/Runtime/Analytics/AnalyticsET/Public/AnalyticsET.h b/Engine/Source/Runtime/Analytics/AnalyticsET/Public/AnalyticsET.h
--- a/Engine/Source/Runtime/Analytics/AnalyticsET/Public/AnalyticsET.h
+++ b/Engine/Source/Runtime/Analytics/AnalyticsET/Public/AnalyticsET.h
@@ -19,7 +19,7 @@ class IAnalyticsProviderET;
 /**
  *  Public implementation of EpicGames.MCP.AnalyticsProvider
  */
-class ANALYTICSET_API FAnalyticsET : public IAnalyticsProviderModule
+class FAnalyticsET : public IAnalyticsProviderModule
 {
 	//--------------------------------------------------------------------------
 	// Module functionality
@@ -119,14 +119,14 @@ public:
 	 * Creates the analytics provider given a configuration delegate.
 	 * The keys required exactly match the field names in the Config object. 
 	 */
-	virtual TSharedPtr<IAnalyticsProvider> CreateAnalyticsProvider(const FAnalyticsProviderConfigurationDelegate& GetConfigValue) const override;
+	ANALYTICSET_API virtual TSharedPtr<IAnalyticsProvider> CreateAnalyticsProvider(const FAnalyticsProviderConfigurationDelegate& GetConfigValue) const override;
 	
 	/** 
 	 * Construct an ET analytics provider directly from a config object.
 	 */
-	virtual TSharedPtr<IAnalyticsProviderET> CreateAnalyticsProvider(const Config& ConfigValues) const;
+	ANALYTICSET_API virtual TSharedPtr<IAnalyticsProviderET> CreateAnalyticsProvider(const Config& ConfigValues) const;
 
 private:
-	virtual void StartupModule() override;
-	virtual void ShutdownModule() override;
+	ANALYTICSET_API virtual void StartupModule() override;
+	ANALYTICSET_API virtual void ShutdownModule() override;
 };
