Router object to connect telemetry producers with consumer that can aggregate, filter, format and send telemetry to endpoints.

diff --git a/Engine/Source/Programs/LowLevelTests/FoundationTests.Build.cs b/Engine/Source/Programs/LowLevelTests/FoundationTests.Build.cs
--- a/Engine/Source/Programs/LowLevelTests/FoundationTests.Build.cs
+++ b/Engine/Source/Programs/LowLevelTests/FoundationTests.Build.cs
@@ -19,7 +19,8 @@ public class FoundationTests : TestModuleRules
 			new string[] {
 				"Core",
 				"Cbor",
-				"CoreUObject"
+				"CoreUObject",
+				"TelemetryUtils"
 			});
 
 		if (Target.bBuildWithEditorOnlyData)
diff --git a/Engine/Source/Runtime/Analytics/TelemetryUtils/Private/TelemetryRouter.cpp b/Engine/Source/Runtime/Analytics/TelemetryUtils/Private/TelemetryRouter.cpp
new file mode 100644
--- /dev/null
+++ b/Engine/Source/Runtime/Analytics/TelemetryUtils/Private/TelemetryRouter.cpp
@@ -0,0 +1,55 @@
+// Copyright Epic Games, Inc. All Rights Reserved.
+ 
+#include "TelemetryRouter.h"
+#include "TelemetryUtils.h"
+#include "Misc/ScopeRWLock.h"
+
+FTelemetryRouter::FTelemetryRouter()
+{
+}
+
+FTelemetryRouter::~FTelemetryRouter()
+{
+}
+
+FTelemetryRouter& FTelemetryRouter::Get()
+{
+    return FTelemetryUtils::GetRouter();
+}
+
+void FTelemetryRouter::ProvideTelemetryInternal(FGuid Key, const void* Data) 
+{
+    FReadScopeLock Lock(SinkLock);
+    TGuardValue Guard(ReentrancyGuard, FPlatformTLS::GetCurrentThreadId());
+    TMap<FDelegateHandle, TFunction<bool(const void*)>>* Sinks = KeyToSinks.Find(Key);
+    if (Sinks)
+    {
+        for (auto It = Sinks->CreateIterator(); It; ++It)
+        {
+            bool bStillBound = It->Value(Data);
+            if (!bStillBound)
+            {
+                It.RemoveCurrent();
+            }
+        }
+    }
+}
+
+void FTelemetryRouter::RegisterTelemetrySinkInternal(FGuid Key, FDelegateHandle InHandle, TFunction<bool(const void*)> Sink)
+{
+    FWriteScopeLock Lock(SinkLock);
+    TGuardValue Guard(ReentrancyGuard, FPlatformTLS::GetCurrentThreadId());
+    TMap<FDelegateHandle, TFunction<bool(const void*)>>& Sinks = KeyToSinks.FindOrAdd(Key);
+    Sinks.Add(InHandle, MoveTemp(Sink));
+}
+
+void FTelemetryRouter::UnregisterTelemetrySinkInternal(FGuid Key, FDelegateHandle InHandle)
+{
+    FWriteScopeLock Lock(SinkLock);
+    TGuardValue Guard(ReentrancyGuard, FPlatformTLS::GetCurrentThreadId());
+    TMap<FDelegateHandle, TFunction<bool(const void*)>>* Sinks = KeyToSinks.Find(Key);
+    if (Sinks)
+    {
+        Sinks->Remove(InHandle);
+    }
+}
\ No newline at end of file
diff --git a/Engine/Source/Runtime/Analytics/TelemetryUtils/Private/TelemetryUtils.cpp b/Engine/Source/Runtime/Analytics/TelemetryUtils/Private/TelemetryUtils.cpp
new file mode 100644
--- /dev/null
+++ b/Engine/Source/Runtime/Analytics/TelemetryUtils/Private/TelemetryUtils.cpp
@@ -0,0 +1,24 @@
+// Copyright Epic Games, Inc. All Rights Reserved.
+ 
+#include "TelemetryUtils.h"
+#include "TelemetryRouter.h"
+
+IMPLEMENT_MODULE( FTelemetryUtils, TelemetryUtils );
+
+FTelemetryUtils::FTelemetryUtils()
+{
+}
+
+FTelemetryUtils::~FTelemetryUtils()
+{
+}
+
+void FTelemetryUtils::StartupModule()
+{
+    Router = MakeUnique<FTelemetryRouter>();
+}
+
+void FTelemetryUtils::ShutdownModule()
+{
+    Router.Reset();
+}
diff --git a/Engine/Source/Runtime/Analytics/TelemetryUtils/Public/TelemetryRouter.h b/Engine/Source/Runtime/Analytics/TelemetryUtils/Public/TelemetryRouter.h
new file mode 100644
--- /dev/null
+++ b/Engine/Source/Runtime/Analytics/TelemetryUtils/Public/TelemetryRouter.h
@@ -0,0 +1,141 @@
+// Copyright Epic Games, Inc. All Rights Reserved.
+
+#pragma once
+
+#include "Delegates/Delegate.h"
+#include "HAL/CriticalSection.h"
+#include "Templates/Models.h"
+#include "Misc/Guid.h"
+
+#include <type_traits>
+
+namespace UE::Telemetry::Private
+{
+    template<typename DATA_TYPE, typename = void>
+    constexpr bool HasTelemetryID = false;
+
+    template<typename DATA_TYPE>
+    constexpr bool HasTelemetryID<
+        DATA_TYPE,
+        std::enable_if_t<std::is_same_v<std::decay_t<decltype(DATA_TYPE::TelemetryID)>, FGuid>>> = true;
+    
+    template <typename DATA_TYPE>
+    auto GetDataKeyInternal()
+    {
+        if constexpr (HasTelemetryID<DATA_TYPE>)
+        {
+            return DATA_TYPE::TelemetryID;
+        }
+        return FGuid{};
+    }
+
+    template <typename DATA_TYPE>
+    FGuid GetDataKey() 
+    {
+        static_assert(HasTelemetryID<DATA_TYPE>, 
+            "Unsupported type for telemetry router, must have a static constexpr FGuid member TelemetryID");
+        // Silence duplicate compiler error that's less helpful than the static assert
+        if constexpr (HasTelemetryID<DATA_TYPE>)
+        {
+            return DATA_TYPE::TelemetryID;
+        }
+        return FGuid{};
+    }
+}
+
+/** 
+ * Provides an interface for routing structured telemetry data between prooducers (engine/editor systems)
+ * and consumers (modules which will format the data for a particular telemetry endpoint)
+ * 
+ * Usage: 
+ *  Declare a unique type for telemetry data. This type needs a static FGuid data member called TelemetryID.
+ *  Provide data to any interested aggregators with the method ProvideTelemetry.
+ *  Aggregators can register a callback for telemetry data with the method OnTelemetry.
+ * 
+ * The guid associated with a type for routing should not be relied on by consumers to be stable across engine versions 
+ * or even across different processes, it is only used for runtime routing purposes. It should not be sent to telemetry endpoints.
+ * 
+ * @thread_safety Functions are internally synchronized. Re-entrant calls to the public API during callbacks are unsupported.
+ *  Data is routed on the same thread it is provided. Users should copy data and schedule tasks to execute on a specified thread 
+ *  if they require that. 
+ */
+class FTelemetryRouter 
+{
+public:
+    FTelemetryRouter();
+    ~FTelemetryRouter();
+    
+    static TELEMETRYUTILS_API FTelemetryRouter& Get();
+
+    /**
+     * Sends data in a type-safe manner to consumers expecting this data type. 
+     * 
+     * @param Data Strongly typed data to be sent. Type must implement a static method called GetTelemetryID returning an FGuid.
+     */
+    template<typename DATA_TYPE>
+    inline void ProvideTelemetry(const DATA_TYPE& Data)
+    {
+        check(ReentrancyGuard != FPlatformTLS::GetCurrentThreadId());
+        ProvideTelemetryInternal(UE::Telemetry::Private::GetDataKey<DATA_TYPE>(), reinterpret_cast<const void*>(&Data));
+    }
+    
+    /** 
+     * Registers a delegate as callback to receive telemetry of a certain type. 
+     *
+     * @return a handle that can be used to unregister this sink later so it is no longer called.
+     */
+    template<typename DATA_TYPE>
+    inline FDelegateHandle OnTelemetry(TDelegate<void(const DATA_TYPE&)> Sink)
+    {
+        check(ReentrancyGuard != FPlatformTLS::GetCurrentThreadId());
+        check(Sink.IsBound());
+        FDelegateHandle Handle = Sink.GetHandle();
+        RegisterTelemetrySinkInternal(UE::Telemetry::Private::GetDataKey<DATA_TYPE>(), Handle,
+            [Sink=MoveTemp(Sink)](const void* Data) -> bool {
+                return Sink.ExecuteIfBound(*reinterpret_cast<const DATA_TYPE*>(Data));
+            });
+        return Handle;
+    }
+
+    /** 
+     * Registers a callable object as a callback receive telemetry of a certain type. 
+     *
+     * @return a handle that can be used to unregister this sink later so it is no longer called.
+     */
+    template<typename DATA_TYPE, typename CALLABLE,
+        typename = std::enable_if_t<std::is_invocable_v<CALLABLE, const DATA_TYPE&>>
+    >
+    inline FDelegateHandle OnTelemetry(CALLABLE&& Sink)
+    {
+        check(ReentrancyGuard != FPlatformTLS::GetCurrentThreadId());
+        FDelegateHandle Handle{ FDelegateHandle::GenerateNewHandle };
+        RegisterTelemetrySinkInternal(UE::Telemetry::Private::GetDataKey<DATA_TYPE>(), Handle,
+            [Sink=MoveTemp(Sink)](const void* Data) -> bool {
+                Sink(*reinterpret_cast<const DATA_TYPE*>(Data));
+                return true;
+            });
+        return Handle;
+    }
+    
+    /** 
+     * Removes a previously registered callback.
+     */
+    template<typename DATA_TYPE>
+    inline void UnregisterTelemetrySink(FDelegateHandle Handle)
+    {
+        check(ReentrancyGuard != FPlatformTLS::GetCurrentThreadId());
+        UnregisterTelemetrySinkInternal(UE::Telemetry::Private::GetDataKey<DATA_TYPE>(), Handle);
+    }
+    
+private:
+    /** 
+     * Register a callback to receive telemetry. The callback should return false if the sink is stale and should be removed from future consideration.
+     */
+    TELEMETRYUTILS_API void RegisterTelemetrySinkInternal(FGuid Key, FDelegateHandle InHandle, TFunction<bool(const void*)> Sink);
+    TELEMETRYUTILS_API void UnregisterTelemetrySinkInternal(FGuid Key, FDelegateHandle InHandle);
+    TELEMETRYUTILS_API void ProvideTelemetryInternal(FGuid Key, const void* Data);
+    
+    FRWLock SinkLock;
+    TMap<FGuid, TMap<FDelegateHandle, TFunction<bool(const void*)>>> KeyToSinks;
+    uint32 ReentrancyGuard = 0;
+};
\ No newline at end of file
diff --git a/Engine/Source/Runtime/Analytics/TelemetryUtils/Public/TelemetryUtils.h b/Engine/Source/Runtime/Analytics/TelemetryUtils/Public/TelemetryUtils.h
new file mode 100644
--- /dev/null
+++ b/Engine/Source/Runtime/Analytics/TelemetryUtils/Public/TelemetryUtils.h
@@ -0,0 +1,44 @@
+// Copyright Epic Games, Inc. All Rights Reserved.
+
+#pragma once
+
+// Module containing utility functions for providing telemetry on the engine and games.
+// Does not integrate with any specific telemetry endpoint.
+// Provides functions for routing data from runtime/editor systems to telemetry sinks in a structured manner.
+
+#include "Modules/ModuleInterface.h"
+#include "Modules/ModuleManager.h"
+
+class FTelemetryRouter;
+
+class FTelemetryUtils : public IModuleInterface
+{
+public:
+    FTelemetryUtils();
+    virtual ~FTelemetryUtils();
+
+    TELEMETRYUTILS_API static inline FTelemetryUtils& Get()
+    {
+        return FModuleManager::LoadModuleChecked<FTelemetryUtils>("TelemetryUtils");
+    }
+    
+    TELEMETRYUTILS_API static inline bool IsAvailable()
+    {
+        return FModuleManager::Get().IsModuleLoaded("TelemetryUtils");
+    }
+    
+    /** 
+     * Returns an object used for routing structured telemetry data between producers and consumers.
+     */
+    TELEMETRYUTILS_API static FTelemetryRouter& GetRouter()
+    {
+        return *FTelemetryUtils::Get().Router.Get();
+    }
+
+private:
+	virtual void StartupModule() override;
+	virtual void ShutdownModule() override;
+    
+private:
+    TUniquePtr<FTelemetryRouter> Router;
+};
\ No newline at end of file
diff --git a/Engine/Source/Runtime/Analytics/TelemetryUtils/TelemetryUtils.Build.cs b/Engine/Source/Runtime/Analytics/TelemetryUtils/TelemetryUtils.Build.cs
new file mode 100644
--- /dev/null
+++ b/Engine/Source/Runtime/Analytics/TelemetryUtils/TelemetryUtils.Build.cs
@@ -0,0 +1,19 @@
+// Copyright Epic Games, Inc. All Rights Reserved.
+
+namespace UnrealBuildTool.Rules
+{
+	public class TelemetryUtils : ModuleRules
+	{
+		public TelemetryUtils(ReadOnlyTargetRules Target) : base(Target)
+		{			
+			PublicDependencyModuleNames.AddRange(
+				new string[]
+				{
+					"Core",
+				}
+				);
+			UnsafeTypeCastWarningLevel = WarningLevel.Error;
+			CppStandard = CppStandardVersion.Cpp20;
+		}
+	}
+}
diff --git a/Engine/Source/Runtime/Analytics/TelemetryUtils/Tests/TelemetryRouterTests.cpp b/Engine/Source/Runtime/Analytics/TelemetryUtils/Tests/TelemetryRouterTests.cpp
new file mode 100644
--- /dev/null
+++ b/Engine/Source/Runtime/Analytics/TelemetryUtils/Tests/TelemetryRouterTests.cpp
@@ -0,0 +1,111 @@
+// Copyright Epic Games, Inc. All Rights Reserved.
+
+#if WITH_LOW_LEVEL_TESTS
+
+#include "TelemetryRouter.h"
+#include "TestHarness.h"
+
+struct FTestTelemetryData
+{
+    int32 Count = 0;
+
+    static constexpr inline FGuid TelemetryID = FGuid(0x4830c7cf, 0x2c184669, 0x89b299c2, 0x3134f2b3);
+};
+
+TEST_CASE("TelemetryUtils::Router::Single", "[Smoke]")
+{
+    FTelemetryRouter Router;
+
+    int32 Increment = 4;
+    FTestTelemetryData DataReceived;
+    FDelegateHandle Handle = Router.OnTelemetry<FTestTelemetryData>(
+        [&DataReceived](const FTestTelemetryData& In) {
+            DataReceived.Count += In.Count;
+        });
+    
+    Router.ProvideTelemetry<FTestTelemetryData>(FTestTelemetryData{ 
+        .Count = Increment
+    });
+    CHECK_EQUAL(DataReceived.Count, Increment);
+    Router.UnregisterTelemetrySink<FTestTelemetryData>(Handle);
+    Router.ProvideTelemetry<FTestTelemetryData>(FTestTelemetryData{
+        .Count = Increment
+    });
+    CHECK_EQUAL(DataReceived.Count, Increment);
+}
+
+TEST_CASE("TelemetryUtils::Router::Multiple", "[Smoke]")
+{
+    FTelemetryRouter Router;
+
+    FTestTelemetryData ReceivedA;
+    FTestTelemetryData ReceivedB;
+    
+    FDelegateHandle HandleA = Router.OnTelemetry<FTestTelemetryData>(
+        [&ReceivedA](const FTestTelemetryData& In) {
+            ReceivedA.Count += In.Count;
+        });
+    FDelegateHandle HandleB = Router.OnTelemetry<FTestTelemetryData>(
+        [&ReceivedB](const FTestTelemetryData& In) {
+            ReceivedB.Count += In.Count;
+        });
+    
+    Router.ProvideTelemetry<FTestTelemetryData>(FTestTelemetryData{ 
+        .Count = 1
+    });
+    CHECK_EQUAL(ReceivedA.Count, 1);
+    CHECK_EQUAL(ReceivedB.Count, 1);
+    Router.UnregisterTelemetrySink<FTestTelemetryData>(HandleA);
+    Router.ProvideTelemetry<FTestTelemetryData>(FTestTelemetryData{
+        .Count = 1
+    });
+    CHECK_EQUAL(ReceivedA.Count, 1);
+    CHECK_EQUAL(ReceivedB.Count, 2);
+}
+
+TEST_CASE("TelemetryUtils::Router::WeakBinding", "[Smoke]")
+{
+    FTelemetryRouter Router;
+    
+    struct FCallee
+    {
+        FCallee()
+        {
+        }
+        ~FCallee()
+        {
+        }
+        
+        void ReceiveTelemetry(const FTestTelemetryData& In)
+        {
+            ReceivedData.Count += In.Count;
+        }
+        
+        FTestTelemetryData ReceivedData;
+    };
+
+    int32 Increment = 4;
+    TWeakPtr<FCallee> WeakCallee;
+    {
+        TSharedRef<FCallee> Callee = MakeShared<FCallee>();
+        FDelegateHandle Handle = Router.OnTelemetry<FTestTelemetryData>(
+            TDelegate<void(const FTestTelemetryData&)>::CreateSP(Callee, &FCallee::ReceiveTelemetry)
+        );
+        
+        Router.ProvideTelemetry<FTestTelemetryData>(FTestTelemetryData{ 
+            .Count = Increment
+        });
+        CHECK_EQUAL(Callee->ReceivedData.Count, Increment);
+        
+        // Delete object without unregistering data sink
+        WeakCallee = Callee;
+        CHECK(WeakCallee.Pin().IsValid());
+    }
+
+    Router.ProvideTelemetry<FTestTelemetryData>(FTestTelemetryData{
+        .Count = Increment
+    });
+    CHECK(!WeakCallee.Pin().IsValid());
+}
+
+#endif
\ No newline at end of file
