Add -telemetrytesting command line argument as an easier-to-remember version of analyticsdisablecaching

diff --git a/Engine/Source/Runtime/Analytics/AnalyticsET/Private/IAnalyticsProviderET.cpp b/Engine/Source/Runtime/Analytics/AnalyticsET/Private/IAnalyticsProviderET.cpp
--- a/Engine/Source/Runtime/Analytics/AnalyticsET/Private/IAnalyticsProviderET.cpp
+++ b/Engine/Source/Runtime/Analytics/AnalyticsET/Private/IAnalyticsProviderET.cpp
@@ -185,6 +185,13 @@ FAnalyticsProviderET::FAnalyticsProviderET(const FAnalyticsET::Config& ConfigVal
 		RetryServers = MakeShared<FHttpRetrySystem::FRetryDomains, ESPMode::ThreadSafe>(MoveTemp(TmpAltAPIServers));
 	}
 
+	const bool bTestingMode = FParse::Param(FCommandLine::Get(), TEXT("TELEMETRYTESTING"));
+	if (bTestingMode)
+	{
+		UE_SET_LOG_VERBOSITY(LogAnalytics, VeryVerbose);
+		bShouldCacheEvents = false;
+	}
+
 	// force very verbose logging if we are force-disabling events.
 	bool bForceDisableCaching = FParse::Param(FCommandLine::Get(), TEXT("ANALYTICSDISABLECACHING"));
 	if (bForceDisableCaching)
