Version of ProvideTelemetry taking a lambda to avoid calculating some statistics if they aren't necessary.

diff --git a/Engine/Source/Runtime/Analytics/TelemetryUtils/Private/TelemetryRouter.cpp b/Engine/Source/Runtime/Analytics/TelemetryUtils/Private/TelemetryRouter.cpp
--- a/Engine/Source/Runtime/Analytics/TelemetryUtils/Private/TelemetryRouter.cpp
+++ b/Engine/Source/Runtime/Analytics/TelemetryUtils/Private/TelemetryRouter.cpp
@@ -17,11 +17,8 @@ FTelemetryRouter& FTelemetryRouter::Get()
     return FTelemetryUtils::GetRouter();
 }
 
-void FTelemetryRouter::ProvideTelemetryInternal(FGuid Key, const void* Data) 
+static void CallSinks(TMap<FDelegateHandle, TFunction<bool(const void*)>>* Sinks, const void* Data)
 {
-    FReadScopeLock Lock(SinkLock);
-    TGuardValue Guard(ReentrancyGuard, FPlatformTLS::GetCurrentThreadId());
-    TMap<FDelegateHandle, TFunction<bool(const void*)>>* Sinks = KeyToSinks.Find(Key);
     if (Sinks)
     {
         for (auto It = Sinks->CreateIterator(); It; ++It)
@@ -35,6 +32,24 @@ void FTelemetryRouter::ProvideTelemetryInternal(FGuid Key, const void* Data)
     }
 }
 
+void FTelemetryRouter::ProvideTelemetryInternal(FGuid Key, const void* Data) 
+{
+    FReadScopeLock Lock(SinkLock);
+    TGuardValue Guard(ReentrancyGuard, FPlatformTLS::GetCurrentThreadId());
+    CallSinks(KeyToSinks.Find(Key), Data);
+}
+
+void FTelemetryRouter::ProvideTelemetryInternal(FGuid Key, TFunctionRef<const void*()> GetData) 
+{
+    FReadScopeLock Lock(SinkLock);
+    TGuardValue Guard(ReentrancyGuard, FPlatformTLS::GetCurrentThreadId());
+    TMap<FDelegateHandle, TFunction<bool(const void*)>>* Sinks = KeyToSinks.Find(Key);
+    if (Sinks)
+    {
+        CallSinks(Sinks, GetData());
+    }
+}
+
 void FTelemetryRouter::RegisterTelemetrySinkInternal(FGuid Key, FDelegateHandle InHandle, TFunction<bool(const void*)> Sink)
 {
     FWriteScopeLock Lock(SinkLock);
diff --git a/Engine/Source/Runtime/Analytics/TelemetryUtils/Public/TelemetryRouter.h b/Engine/Source/Runtime/Analytics/TelemetryUtils/Public/TelemetryRouter.h
--- a/Engine/Source/Runtime/Analytics/TelemetryUtils/Public/TelemetryRouter.h
+++ b/Engine/Source/Runtime/Analytics/TelemetryUtils/Public/TelemetryRouter.h
@@ -78,6 +78,25 @@ public:
         check(ReentrancyGuard != FPlatformTLS::GetCurrentThreadId());
         ProvideTelemetryInternal(UE::Telemetry::Private::GetDataKey<DATA_TYPE>(), reinterpret_cast<const void*>(&Data));
     }
+
+    /**
+     * Sends data in a type-safe manner to consumers expecting this data type. 
+     * This version takes a lambda so that expensive data can only be calculated if a registered callback will consume it.
+     * 
+     * @param Data Strongly typed data to be sent. Type must implement a static method called GetTelemetryID returning an FGuid.
+     */
+    template<typename DATA_TYPE, typename CALLABLE
+        UE_REQUIRES(std::is_same_v<std::invoke_result_t<CALLABLE>, DATA_TYPE>)
+    >
+    inline void ProvideTelemetryWith(CALLABLE&& Callable)
+    {
+        check(ReentrancyGuard != FPlatformTLS::GetCurrentThreadId());
+        TOptional<DATA_TYPE> Value;
+        ProvideTelemetryInternal(UE::Telemetry::Private::GetDataKey<DATA_TYPE>(), [&Value, &Callable](){
+            Value = Callable();
+            return reinterpret_cast<const void*>(&Value.GetValue());
+        });
+    }
     
     /** 
      * Registers a delegate as callback to receive telemetry of a certain type. 
@@ -134,6 +153,7 @@ private:
     TELEMETRYUTILS_API void RegisterTelemetrySinkInternal(FGuid Key, FDelegateHandle InHandle, TFunction<bool(const void*)> Sink);
     TELEMETRYUTILS_API void UnregisterTelemetrySinkInternal(FGuid Key, FDelegateHandle InHandle);
     TELEMETRYUTILS_API void ProvideTelemetryInternal(FGuid Key, const void* Data);
+    TELEMETRYUTILS_API void ProvideTelemetryInternal(FGuid Key, TFunctionRef<const void*()> GetData);
     
     FRWLock SinkLock;
     TMap<FGuid, TMap<FDelegateHandle, TFunction<bool(const void*)>>> KeyToSinks;
diff --git a/Engine/Source/Runtime/Analytics/TelemetryUtils/Tests/TelemetryRouterTests.cpp b/Engine/Source/Runtime/Analytics/TelemetryUtils/Tests/TelemetryRouterTests.cpp
--- a/Engine/Source/Runtime/Analytics/TelemetryUtils/Tests/TelemetryRouterTests.cpp
+++ b/Engine/Source/Runtime/Analytics/TelemetryUtils/Tests/TelemetryRouterTests.cpp
@@ -108,4 +108,27 @@ TEST_CASE("TelemetryUtils::Router::WeakBinding", "[Smoke]")
     CHECK(!WeakCallee.Pin().IsValid());
 }
 
+TEST_CASE("TelemetryUtils::Router::ProvideWithLambda", "[Smoke]")
+{
+    FTelemetryRouter Router;
+
+    constexpr int32 Increment = 4;
+    FTestTelemetryData DataReceived;
+    FDelegateHandle Handle = Router.OnTelemetry<FTestTelemetryData>(
+        [&DataReceived](const FTestTelemetryData& In) {
+            DataReceived.Count += In.Count;
+        });
+    
+    Router.ProvideTelemetryWith<FTestTelemetryData>([]() -> FTestTelemetryData { 
+        return FTestTelemetryData{ .Count = Increment };
+    });
+    CHECK_EQUAL(DataReceived.Count, Increment);
+    Router.UnregisterTelemetrySink<FTestTelemetryData>(Handle);
+    Router.ProvideTelemetryWith<FTestTelemetryData>([]() -> FTestTelemetryData { 
+        return FTestTelemetryData{ .Count = Increment };
+    });
+    CHECK_EQUAL(DataReceived.Count, Increment);
+}
+
+
 #endif
\ No newline at end of file
