Fix for Analytics tick float overflow, restore original intent

diff --git a/Engine/Source/Runtime/Analytics/AnalyticsET/Private/IAnalyticsProviderET.cpp b/Engine/Source/Runtime/Analytics/AnalyticsET/Private/IAnalyticsProviderET.cpp
--- a/Engine/Source/Runtime/Analytics/AnalyticsET/Private/IAnalyticsProviderET.cpp
+++ b/Engine/Source/Runtime/Analytics/AnalyticsET/Private/IAnalyticsProviderET.cpp
@@ -347,7 +347,7 @@ bool FAnalyticsProviderET::Tick(float DeltaSeconds)
 			// try to keep on the same cadence when flushing, since we could miss our window by several frames.
 			if (!bHadFlushesQueued && Now >= NextEventFlushTime)
 			{
-				const float Multiplier = (FloatCastChecked<float>(Now - NextEventFlushTime, 1./16.) / FlushIntervalSec) + 1.f;
+				const double Multiplier = FMath::Floor((Now - NextEventFlushTime) / FlushIntervalSec) + 1.;
 				NextEventFlushTime += Multiplier * FlushIntervalSec;
 			}
 		}
