Add size validation to sometimes catch copy-pasted guids in telemetry event declarations.
Use FMemoryView instead of const void*.

diff --git a/Engine/Source/Runtime/Analytics/TelemetryUtils/Private/TelemetryRouter.cpp b/Engine/Source/Runtime/Analytics/TelemetryUtils/Private/TelemetryRouter.cpp
--- a/Engine/Source/Runtime/Analytics/TelemetryUtils/Private/TelemetryRouter.cpp
+++ b/Engine/Source/Runtime/Analytics/TelemetryUtils/Private/TelemetryRouter.cpp
@@ -17,14 +17,20 @@ FTelemetryRouter& FTelemetryRouter::Get()
     return FTelemetryUtils::GetRouter();
 }
 
-void FTelemetryRouter::ProvideTelemetryInternal(FGuid Key, const void* Data) 
+void FTelemetryRouter::ProvideTelemetryInternal(FGuid Key, FMemoryView Data) 
 {
     FReadScopeLock Lock(SinkLock);
     TGuardValue Guard(ReentrancyGuard, FPlatformTLS::GetCurrentThreadId());
-    TMap<FDelegateHandle, TFunction<bool(const void*)>>* Sinks = KeyToSinks.Find(Key);
+    FSinkSet* Sinks = KeyToSinks.Find(Key);
     if (Sinks)
     {
-        for (auto It = Sinks->CreateIterator(); It; ++It)
+        if (!ensureMsgf(Sinks->DataSize == Data.GetSize(), TEXT("Size mismatch mismatch with registered telemetry sink for guid %s. %" SIZE_T_X_FMT " vs %" SIZE_T_X_FMT), 
+            *Key.ToString(), Sinks->DataSize, Data.GetSize() ))
+        {
+            return;
+        }
+
+        for (auto It = Sinks->Delegates.CreateIterator(); It; ++It)
         {
             bool bStillBound = It->Value(Data);
             if (!bStillBound)
@@ -35,21 +41,30 @@ void FTelemetryRouter::ProvideTelemetryInternal(FGuid Key, const void* Data)
     }
 }
 
-void FTelemetryRouter::RegisterTelemetrySinkInternal(FGuid Key, FDelegateHandle InHandle, TFunction<bool(const void*)> Sink)
+void FTelemetryRouter::RegisterTelemetrySinkInternal(FGuid Key, SIZE_T Size, FDelegateHandle InHandle, TFunction<bool(FMemoryView)> Sink)
 {
     FWriteScopeLock Lock(SinkLock);
     TGuardValue Guard(ReentrancyGuard, FPlatformTLS::GetCurrentThreadId());
-    TMap<FDelegateHandle, TFunction<bool(const void*)>>& Sinks = KeyToSinks.FindOrAdd(Key);
-    Sinks.Add(InHandle, MoveTemp(Sink));
+    FTelemetryRouter::FSinkSet* Sinks = KeyToSinks.Find(Key);
+    if (!Sinks)
+    {
+        Sinks = &KeyToSinks.Add(Key, FSinkSet(Size));
+    }
+    if (!ensureMsgf(Sinks->DataSize == Size, TEXT("Size mismatch mismatch with registered telemetry sink for guid %s. %" SIZE_T_X_FMT " vs %" SIZE_T_X_FMT), 
+        *Key.ToString(), Sinks->DataSize, Size))
+    {
+        return;
+    }
+    Sinks->Delegates.Add(InHandle, MoveTemp(Sink));
 }
 
 void FTelemetryRouter::UnregisterTelemetrySinkInternal(FGuid Key, FDelegateHandle InHandle)
 {
     FWriteScopeLock Lock(SinkLock);
     TGuardValue Guard(ReentrancyGuard, FPlatformTLS::GetCurrentThreadId());
-    TMap<FDelegateHandle, TFunction<bool(const void*)>>* Sinks = KeyToSinks.Find(Key);
+    FTelemetryRouter::FSinkSet* Sinks = KeyToSinks.Find(Key);
     if (Sinks)
     {
-        Sinks->Remove(InHandle);
+        Sinks->Delegates.Remove(InHandle);
     }
 }
\ No newline at end of file
diff --git a/Engine/Source/Runtime/Analytics/TelemetryUtils/Public/TelemetryRouter.h b/Engine/Source/Runtime/Analytics/TelemetryUtils/Public/TelemetryRouter.h
--- a/Engine/Source/Runtime/Analytics/TelemetryUtils/Public/TelemetryRouter.h
+++ b/Engine/Source/Runtime/Analytics/TelemetryUtils/Public/TelemetryRouter.h
@@ -5,6 +5,7 @@
 #include "Delegates/Delegate.h"
 #include "HAL/CriticalSection.h"
 #include "Templates/Models.h"
+#include "Memory/MemoryView.h"
 #include "Misc/Guid.h"
 
 #include <type_traits>
@@ -76,7 +77,7 @@ public:
     inline void ProvideTelemetry(const DATA_TYPE& Data)
     {
         check(ReentrancyGuard != FPlatformTLS::GetCurrentThreadId());
-        ProvideTelemetryInternal(UE::Telemetry::Private::GetDataKey<DATA_TYPE>(), reinterpret_cast<const void*>(&Data));
+        ProvideTelemetryInternal(UE::Telemetry::Private::GetDataKey<DATA_TYPE>(), MakeMemoryView(&Data, sizeof(DATA_TYPE)));
     }
     
     /** 
@@ -90,9 +91,9 @@ public:
         check(ReentrancyGuard != FPlatformTLS::GetCurrentThreadId());
         check(Sink.IsBound());
         FDelegateHandle Handle = Sink.GetHandle();
-        RegisterTelemetrySinkInternal(UE::Telemetry::Private::GetDataKey<DATA_TYPE>(), Handle,
-            [Sink=MoveTemp(Sink)](const void* Data) -> bool {
-                return Sink.ExecuteIfBound(*reinterpret_cast<const DATA_TYPE*>(Data));
+        RegisterTelemetrySinkInternal(UE::Telemetry::Private::GetDataKey<DATA_TYPE>(), sizeof(DATA_TYPE), Handle,
+            [Sink=MoveTemp(Sink)](FMemoryView Data) -> bool {
+                return Sink.ExecuteIfBound(*reinterpret_cast<const DATA_TYPE*>(Data.GetData()));
             });
         return Handle;
     }
@@ -109,9 +110,9 @@ public:
     {
         check(ReentrancyGuard != FPlatformTLS::GetCurrentThreadId());
         FDelegateHandle Handle{ FDelegateHandle::GenerateNewHandle };
-        RegisterTelemetrySinkInternal(UE::Telemetry::Private::GetDataKey<DATA_TYPE>(), Handle,
-            [Sink=MoveTemp(Sink)](const void* Data) -> bool {
-                Sink(*reinterpret_cast<const DATA_TYPE*>(Data));
+        RegisterTelemetrySinkInternal(UE::Telemetry::Private::GetDataKey<DATA_TYPE>(), sizeof(DATA_TYPE), Handle,
+            [Sink=MoveTemp(Sink)](FMemoryView Data) -> bool {
+                Sink(*reinterpret_cast<const DATA_TYPE*>(Data.GetData()));
                 return true;
             });
         return Handle;
@@ -131,11 +132,23 @@ private:
     /** 
      * Register a callback to receive telemetry. The callback should return false if the sink is stale and should be removed from future consideration.
      */
-    TELEMETRYUTILS_API void RegisterTelemetrySinkInternal(FGuid Key, FDelegateHandle InHandle, TFunction<bool(const void*)> Sink);
+    TELEMETRYUTILS_API void RegisterTelemetrySinkInternal(FGuid Key, SIZE_T ExpectedSize, FDelegateHandle InHandle, TFunction<bool(FMemoryView)> Sink);
     TELEMETRYUTILS_API void UnregisterTelemetrySinkInternal(FGuid Key, FDelegateHandle InHandle);
-    TELEMETRYUTILS_API void ProvideTelemetryInternal(FGuid Key, const void* Data);
+    TELEMETRYUTILS_API void ProvideTelemetryInternal(FGuid Key, FMemoryView Data);
     
     FRWLock SinkLock;
-    TMap<FGuid, TMap<FDelegateHandle, TFunction<bool(const void*)>>> KeyToSinks;
+    
+    struct FSinkSet
+    {
+        SIZE_T DataSize;
+        TMap<FDelegateHandle, TFunction<bool(FMemoryView)>> Delegates;
+        
+        FSinkSet(SIZE_T InDataSize)
+            : DataSize(InDataSize)
+        {
+        }
+    };
+
+    TMap<FGuid, FSinkSet> KeyToSinks;
     uint32 ReentrancyGuard = 0;
 };
\ No newline at end of file
