Analytics: Set PayloadFlushTimeSecForWarning to a larger value when running in commandlets because hitches in commandlets are unimportant.
#rnx

diff --git a/Engine/Source/Runtime/Analytics/AnalyticsET/Private/AnalyticsET.cpp b/Engine/Source/Runtime/Analytics/AnalyticsET/Private/AnalyticsET.cpp
--- a/Engine/Source/Runtime/Analytics/AnalyticsET/Private/AnalyticsET.cpp
+++ b/Engine/Source/Runtime/Analytics/AnalyticsET/Private/AnalyticsET.cpp
@@ -6,6 +6,7 @@
 #include "HttpModule.h"
 #include "Analytics.h"
 #include "AnalyticsPerfTracker.h"
+#include "AnalyticsProviderETEventCache.h"
 
 IMPLEMENT_MODULE( FAnalyticsET, AnalyticsET );
 
@@ -13,6 +14,7 @@ void FAnalyticsET::StartupModule()
 {
 	// Make sure http is loaded so that we can flush events during module shutdown
 	FModuleManager::LoadModuleChecked<FHttpModule>("HTTP");
+	FAnalyticsProviderETEventCache::OnStartupModule();
 }
 
 void FAnalyticsET::ShutdownModule()
diff --git a/Engine/Source/Runtime/Analytics/AnalyticsET/Private/AnalyticsProviderETEventCache.cpp b/Engine/Source/Runtime/Analytics/AnalyticsET/Private/AnalyticsProviderETEventCache.cpp
--- a/Engine/Source/Runtime/Analytics/AnalyticsET/Private/AnalyticsProviderETEventCache.cpp
+++ b/Engine/Source/Runtime/Analytics/AnalyticsET/Private/AnalyticsProviderETEventCache.cpp
@@ -22,7 +22,12 @@ namespace EventCacheStatic
 		TEXT("Percentage of the maximum payload for an EventCache that will trigger a warning message, listing the events in the payload. This is intended to be used to investigate spammy or slow telemetry.")
 	);
 
-	static float PayloadFlushTimeSecForWarning = 0.001f;
+	static float PayloadFlushTimeSecForWarning =
+#if !WITH_EDITOR
+		.001f;
+#else
+		-1.0f; // Initialized in OnStartupModule below
+#endif
 	FAutoConsoleVariableRef CvarPayloadFlushTimeSecForWarning(
 		TEXT("AnalyticsET.PayloadFlushTimeSecForWarning"),
 		PayloadFlushTimeSecForWarning,
@@ -187,6 +192,19 @@ namespace EventCacheStatic
 	}
 }
 
+ANALYTICSET_API void FAnalyticsProviderETEventCache::OnStartupModule()
+{
+#if WITH_EDITOR
+	// Set some performance configuration values that have different defaults depending upon
+	// commandline (editor versus commandlet)
+	bool bCommandlet = IsRunningCommandlet();
+	if (EventCacheStatic::PayloadFlushTimeSecForWarning < 0)
+	{
+		EventCacheStatic::PayloadFlushTimeSecForWarning = !bCommandlet ? 0.001f : 1.0f;
+	}
+#endif
+}
+
 FAnalyticsProviderETEventCache::FAnalyticsProviderETEventCache(int32 InMaximumPayloadSize, int32 InPreallocatedPayloadSize)
 : MaximumPayloadSize(InMaximumPayloadSize)
 , PreallocatedPayloadSize(InPreallocatedPayloadSize)
diff --git a/Engine/Source/Runtime/Analytics/AnalyticsET/Public/AnalyticsProviderETEventCache.h b/Engine/Source/Runtime/Analytics/AnalyticsET/Public/AnalyticsProviderETEventCache.h
--- a/Engine/Source/Runtime/Analytics/AnalyticsET/Public/AnalyticsProviderETEventCache.h
+++ b/Engine/Source/Runtime/Analytics/AnalyticsET/Public/AnalyticsProviderETEventCache.h
@@ -110,6 +110,8 @@ public:
 		FScopeLock ScopedLock;
 	};
 
+	static ANALYTICSET_API void OnStartupModule();
+
 private:
 	ANALYTICSET_API void QueueFlush();
 
