AnalyticsProviderETEventCache: Add more feedback to the warning message "EventCache took too long to flush."
#rnx

diff --git a/Engine/Source/Runtime/Analytics/AnalyticsET/Private/AnalyticsProviderETEventCache.cpp b/Engine/Source/Runtime/Analytics/AnalyticsET/Private/AnalyticsProviderETEventCache.cpp
--- a/Engine/Source/Runtime/Analytics/AnalyticsET/Private/AnalyticsProviderETEventCache.cpp
+++ b/Engine/Source/Runtime/Analytics/AnalyticsET/Private/AnalyticsProviderETEventCache.cpp
@@ -402,10 +402,28 @@ void FAnalyticsProviderETEventCache::QueueFlush()
 	}
 
 	// see if it took too long or we have a really large payload. If so, log out the events.
-	const double EndTime = FPlatformTime::Seconds();
-	if ((EndTime - StartTime) > EventCacheStatic::PayloadFlushTimeSecForWarning || CachedEventUTF8Stream.Num() > (int32)((float)MaximumPayloadSize * EventCacheStatic::PayloadPercentageOfMaxForWarning))
+	if (CachedEventUTF8Stream.Num() > (int32)((float)MaximumPayloadSize * EventCacheStatic::PayloadPercentageOfMaxForWarning))
 	{
-		UE_LOG(LogAnalytics, Warning, TEXT("EventCache either took too long to flush (%.3f ms) or had a very large payload (%.3f KB, %d events). Listing events in the payload for investigation:"), (EndTime-StartTime) * 1000, (float)CachedEventUTF8Stream.Num() / 1024.f, CachedEventEntries.Num());
+		UE_LOG(LogAnalytics, Warning, TEXT("EventCache payload exceeded the maximum allowed size (%.3f KB > %.3f KB). Payload size: %.3f, %d events. Listing events in the payload for investigation:"),
+			(float)CachedEventUTF8Stream.Num() / 1024.f,
+			((float)MaximumPayloadSize * EventCacheStatic::PayloadPercentageOfMaxForWarning) / 1024.f,
+			(float)CachedEventUTF8Stream.Num() / 1024.f, CachedEventEntries.Num());
+		for (const FAnalyticsEventEntry& Entry : CachedEventEntries)
+		{
+			UE_LOG(LogAnalytics, Warning, TEXT("    %s,%d"), *Entry.EventName, Entry.EventSizeChars);
+		}
+	}
+	else if (EventCacheStatic::PayloadFlushTimeSecForWarning < 0)
+	{
+		// Send the callstack for this case, which we can do by logging an ensure
+		ensureMsgf(false, TEXT("QueueFlush called before EventCacheStatic::PayloadFlushTimeSecForWarning was initialized."));
+	}
+	else if (const double EndTime = FPlatformTime::Seconds();
+			 (EndTime - StartTime) > EventCacheStatic::PayloadFlushTimeSecForWarning)
+	{
+		UE_LOG(LogAnalytics, Warning, TEXT("EventCache took too long to flush (%.3f ms > %.3f ms). Payload size: %.3f KB, %d events. Listing events in the payload for investigation:"),
+			(EndTime - StartTime) * 1000, EventCacheStatic::PayloadFlushTimeSecForWarning * 1000,
+			(float)CachedEventUTF8Stream.Num() / 1024.f, CachedEventEntries.Num());
 		for (const FAnalyticsEventEntry& Entry : CachedEventEntries)
 		{
 			UE_LOG(LogAnalytics, Warning, TEXT("    %s,%d"), *Entry.EventName, Entry.EventSizeChars);
