AnalyticsProviderETEventCache: Change the automation test from SmokeFilter (runs early on startup) to EngineFilter (runs only when requested). In monolithic editor configurations, running at SmokeFilter time occurs before StartupModule is called, and the test requires some code from StartupModule.
#rnx

diff --git a/Engine/Source/Runtime/Analytics/AnalyticsET/Private/AnalyticsProviderETEventCache.cpp b/Engine/Source/Runtime/Analytics/AnalyticsET/Private/AnalyticsProviderETEventCache.cpp
--- a/Engine/Source/Runtime/Analytics/AnalyticsET/Private/AnalyticsProviderETEventCache.cpp
+++ b/Engine/Source/Runtime/Analytics/AnalyticsET/Private/AnalyticsProviderETEventCache.cpp
@@ -483,7 +483,7 @@ int32 FAnalyticsProviderETEventCache::GetSetPreallocatedPayloadSize() const
 #include "Misc/AutomationTest.h"
 #include <limits>
 
-IMPLEMENT_SIMPLE_AUTOMATION_TEST(FAnalyticsProviderETEventCacheTest, "System.Analytics.AnalyticsETEventCache", EAutomationTestFlags::ApplicationContextMask | EAutomationTestFlags::SmokeFilter)
+IMPLEMENT_SIMPLE_AUTOMATION_TEST(FAnalyticsProviderETEventCacheTest, "System.Analytics.AnalyticsETEventCache", EAutomationTestFlags::ApplicationContextMask | EAutomationTestFlags::EngineFilter)
 bool FAnalyticsProviderETEventCacheTest::RunTest(const FString& Parameters)
 {
 	// Zero out the DateOffset so we can test against constant strings.
