Change analytics testing to not spew warnings if the flush takes too long. Still spews info logging to help profilers track down responsible parties. Will spew warnings when payloads are too large because those issues should be addressed immediately in any tool.

#tests Compile

diff --git a/Engine/Source/Runtime/Analytics/AnalyticsET/Private/AnalyticsProviderETEventCache.cpp b/Engine/Source/Runtime/Analytics/AnalyticsET/Private/AnalyticsProviderETEventCache.cpp
--- a/Engine/Source/Runtime/Analytics/AnalyticsET/Private/AnalyticsProviderETEventCache.cpp
+++ b/Engine/Source/Runtime/Analytics/AnalyticsET/Private/AnalyticsProviderETEventCache.cpp
@@ -22,12 +22,7 @@ namespace EventCacheStatic
 		TEXT("Percentage of the maximum payload for an EventCache that will trigger a warning message, listing the events in the payload. This is intended to be used to investigate spammy or slow telemetry.")
 	);
 
-	static float PayloadFlushTimeSecForWarning =
-#if !WITH_EDITOR
-		.001f;
-#else
-		-1.0f; // Initialized in OnStartupModule below
-#endif
+	static float PayloadFlushTimeSecForWarning = 0.001f;
 	FAutoConsoleVariableRef CvarPayloadFlushTimeSecForWarning(
 		TEXT("AnalyticsET.PayloadFlushTimeSecForWarning"),
 		PayloadFlushTimeSecForWarning,
@@ -194,15 +189,6 @@ namespace EventCacheStatic
 
 ANALYTICSET_API void FAnalyticsProviderETEventCache::OnStartupModule()
 {
-#if WITH_EDITOR
-	// Set some performance configuration values that have different defaults depending upon
-	// commandline (editor versus commandlet)
-	bool bCommandlet = IsRunningCommandlet();
-	if (EventCacheStatic::PayloadFlushTimeSecForWarning < 0)
-	{
-		EventCacheStatic::PayloadFlushTimeSecForWarning = !bCommandlet ? 0.001f : 1.0f;
-	}
-#endif
 }
 
 FAnalyticsProviderETEventCache::FAnalyticsProviderETEventCache(int32 InMaximumPayloadSize, int32 InPreallocatedPayloadSize)
@@ -402,31 +388,33 @@ void FAnalyticsProviderETEventCache::QueueFlush()
 	}
 
 	// see if it took too long or we have a really large payload. If so, log out the events.
-	if (CachedEventUTF8Stream.Num() > (int32)((float)MaximumPayloadSize * EventCacheStatic::PayloadPercentageOfMaxForWarning))
+	const double EndTime = FPlatformTime::Seconds();
+	const bool bPlayloadTooLarge = CachedEventUTF8Stream.Num() > (int32)((float)MaximumPayloadSize * EventCacheStatic::PayloadPercentageOfMaxForWarning);
+	const bool bTookTooLongToFlush = (EndTime - StartTime) > EventCacheStatic::PayloadFlushTimeSecForWarning;
+	if (bPlayloadTooLarge)
 	{
-		UE_LOG(LogAnalytics, Warning, TEXT("EventCache payload exceeded the maximum allowed size (%.3f KB > %.3f KB). Payload size: %.3f, %d events. Listing events in the payload for investigation:"),
+		
+		UE_LOG(LogAnalytics, Warning, TEXT("EventCache payload exceeded the maximum allowed size (%.3f KB > %.3f KB), containing %d events. Listing events in the payload for investigation:"),
 			(float)CachedEventUTF8Stream.Num() / 1024.f,
 			((float)MaximumPayloadSize * EventCacheStatic::PayloadPercentageOfMaxForWarning) / 1024.f,
-			(float)CachedEventUTF8Stream.Num() / 1024.f, CachedEventEntries.Num());
+			CachedEventEntries.Num());
 		for (const FAnalyticsEventEntry& Entry : CachedEventEntries)
 		{
 			UE_LOG(LogAnalytics, Warning, TEXT("    %s,%d"), *Entry.EventName, Entry.EventSizeChars);
 		}
 	}
-	else if (EventCacheStatic::PayloadFlushTimeSecForWarning < 0)
-	{
-		// Send the callstack for this case, which we can do by logging an ensure
-		ensureMsgf(false, TEXT("QueueFlush called before EventCacheStatic::PayloadFlushTimeSecForWarning was initialized."));
-	}
-	else if (const double EndTime = FPlatformTime::Seconds();
-			 (EndTime - StartTime) > EventCacheStatic::PayloadFlushTimeSecForWarning)
+	// If the event took too long to flush, this may cause it to come up during profiling sessions. But generally, the problem is not with the telemetry code,
+	// the problem is with Events that are trying to send too much data. List the events here to make it a bit easier to track down the responsible party for the slow telemetry.
+	// Don't log at warning level because a lot automated tools don't care if telemetry flushes slowly, and it may happen in practice, and those tools will also error and
+	// break the build if they detect warnings or errors.
+	else if (bTookTooLongToFlush)
 	{
-		UE_LOG(LogAnalytics, Warning, TEXT("EventCache took too long to flush (%.3f ms > %.3f ms). Payload size: %.3f KB, %d events. Listing events in the payload for investigation:"),
+		UE_LOG(LogAnalytics, Display, TEXT("EventCache took too long to flush (%.3f ms > %.3f ms). Payload size: %.3f KB, %d events. Listing events in the payload for investigation:"),
 			(EndTime - StartTime) * 1000, EventCacheStatic::PayloadFlushTimeSecForWarning * 1000,
 			(float)CachedEventUTF8Stream.Num() / 1024.f, CachedEventEntries.Num());
 		for (const FAnalyticsEventEntry& Entry : CachedEventEntries)
 		{
-			UE_LOG(LogAnalytics, Warning, TEXT("    %s,%d"), *Entry.EventName, Entry.EventSizeChars);
+			UE_LOG(LogAnalytics, Display, TEXT("    %s,%d"), *Entry.EventName, Entry.EventSizeChars);
 		}
 	}
 
@@ -483,7 +471,7 @@ int32 FAnalyticsProviderETEventCache::GetSetPreallocatedPayloadSize() const
 #include "Misc/AutomationTest.h"
 #include <limits>
 
-IMPLEMENT_SIMPLE_AUTOMATION_TEST(FAnalyticsProviderETEventCacheTest, "System.Analytics.AnalyticsETEventCache", EAutomationTestFlags::ApplicationContextMask | EAutomationTestFlags::EngineFilter)
+IMPLEMENT_SIMPLE_AUTOMATION_TEST(FAnalyticsProviderETEventCacheTest, "System.Analytics.AnalyticsETEventCache", EAutomationTestFlags::ApplicationContextMask | EAutomationTestFlags::SmokeFilter)
 bool FAnalyticsProviderETEventCacheTest::RunTest(const FString& Parameters)
 {
 	// Zero out the DateOffset so we can test against constant strings.
