TelemetryRouter: Temporarily disable ReentrancyGuard until we have time to make it threadsafe.
#rnx

diff --git a/Engine/Source/Runtime/Analytics/TelemetryUtils/Private/TelemetryRouter.cpp b/Engine/Source/Runtime/Analytics/TelemetryUtils/Private/TelemetryRouter.cpp
--- a/Engine/Source/Runtime/Analytics/TelemetryUtils/Private/TelemetryRouter.cpp
+++ b/Engine/Source/Runtime/Analytics/TelemetryUtils/Private/TelemetryRouter.cpp
@@ -20,7 +20,7 @@ FTelemetryRouter& FTelemetryRouter::Get()
 void FTelemetryRouter::ProvideTelemetryInternal(FGuid Key, FMemoryView Data) 
 {
     FReadScopeLock Lock(SinkLock);
-    TGuardValue Guard(ReentrancyGuard, FPlatformTLS::GetCurrentThreadId());
+    //TGuardValue Guard(ReentrancyGuard, FPlatformTLS::GetCurrentThreadId());
     FSinkSet* Sinks = KeyToSinks.Find(Key);
     if (Sinks)
     {
@@ -44,7 +44,7 @@ void FTelemetryRouter::ProvideTelemetryInternal(FGuid Key, FMemoryView Data)
 void FTelemetryRouter::RegisterTelemetrySinkInternal(FGuid Key, SIZE_T Size, FDelegateHandle InHandle, TFunction<bool(FMemoryView)> Sink)
 {
     FWriteScopeLock Lock(SinkLock);
-    TGuardValue Guard(ReentrancyGuard, FPlatformTLS::GetCurrentThreadId());
+    //TGuardValue Guard(ReentrancyGuard, FPlatformTLS::GetCurrentThreadId());
     FTelemetryRouter::FSinkSet* Sinks = KeyToSinks.Find(Key);
     if (!Sinks)
     {
@@ -61,7 +61,7 @@ void FTelemetryRouter::RegisterTelemetrySinkInternal(FGuid Key, SIZE_T Size, FDe
 void FTelemetryRouter::UnregisterTelemetrySinkInternal(FGuid Key, FDelegateHandle InHandle)
 {
     FWriteScopeLock Lock(SinkLock);
-    TGuardValue Guard(ReentrancyGuard, FPlatformTLS::GetCurrentThreadId());
+    //TGuardValue Guard(ReentrancyGuard, FPlatformTLS::GetCurrentThreadId());
     FTelemetryRouter::FSinkSet* Sinks = KeyToSinks.Find(Key);
     if (Sinks)
     {
diff --git a/Engine/Source/Runtime/Analytics/TelemetryUtils/Public/TelemetryRouter.h b/Engine/Source/Runtime/Analytics/TelemetryUtils/Public/TelemetryRouter.h
--- a/Engine/Source/Runtime/Analytics/TelemetryUtils/Public/TelemetryRouter.h
+++ b/Engine/Source/Runtime/Analytics/TelemetryUtils/Public/TelemetryRouter.h
@@ -76,7 +76,7 @@ public:
     template<typename DATA_TYPE>
     inline void ProvideTelemetry(const DATA_TYPE& Data)
     {
-        check(ReentrancyGuard != FPlatformTLS::GetCurrentThreadId());
+        //check(ReentrancyGuard != FPlatformTLS::GetCurrentThreadId());
         ProvideTelemetryInternal(UE::Telemetry::Private::GetDataKey<DATA_TYPE>(), MakeMemoryView(&Data, sizeof(DATA_TYPE)));
     }
     
@@ -88,7 +88,7 @@ public:
     template<typename DATA_TYPE>
     inline FDelegateHandle OnTelemetry(TDelegate<void(const DATA_TYPE&)> Sink)
     {
-        check(ReentrancyGuard != FPlatformTLS::GetCurrentThreadId());
+        //check(ReentrancyGuard != FPlatformTLS::GetCurrentThreadId());
         check(Sink.IsBound());
         FDelegateHandle Handle = Sink.GetHandle();
         RegisterTelemetrySinkInternal(UE::Telemetry::Private::GetDataKey<DATA_TYPE>(), sizeof(DATA_TYPE), Handle,
@@ -108,7 +108,7 @@ public:
     >
     inline FDelegateHandle OnTelemetry(CALLABLE&& Sink)
     {
-        check(ReentrancyGuard != FPlatformTLS::GetCurrentThreadId());
+        //check(ReentrancyGuard != FPlatformTLS::GetCurrentThreadId());
         FDelegateHandle Handle{ FDelegateHandle::GenerateNewHandle };
         RegisterTelemetrySinkInternal(UE::Telemetry::Private::GetDataKey<DATA_TYPE>(), sizeof(DATA_TYPE), Handle,
             [Sink=MoveTemp(Sink)](FMemoryView Data) -> bool {
@@ -124,7 +124,7 @@ public:
     template<typename DATA_TYPE>
     inline void UnregisterTelemetrySink(FDelegateHandle Handle)
     {
-        check(ReentrancyGuard != FPlatformTLS::GetCurrentThreadId());
+        //check(ReentrancyGuard != FPlatformTLS::GetCurrentThreadId());
         UnregisterTelemetrySinkInternal(UE::Telemetry::Private::GetDataKey<DATA_TYPE>(), Handle);
     }
     
@@ -150,5 +150,7 @@ private:
     };
 
     TMap<FGuid, FSinkSet> KeyToSinks;
-    uint32 ReentrancyGuard = 0;
+	// Disabling ReentrancyGuard since it is not threadsafe, and a race condition can cause it to be set to Thread A again after Thread A raced
+	// Thread B and Thread A returns and then Thread B returns. Need to reconsider whether there is a way to cheaply prevent reentrancy.
+    //uint32 ReentrancyGuard = 0;
 };
\ No newline at end of file
