Re-implement re-entrancy check using thread-local bool.

diff --git a/Engine/Source/Runtime/Analytics/TelemetryUtils/Private/TelemetryRouter.cpp b/Engine/Source/Runtime/Analytics/TelemetryUtils/Private/TelemetryRouter.cpp
--- a/Engine/Source/Runtime/Analytics/TelemetryUtils/Private/TelemetryRouter.cpp
+++ b/Engine/Source/Runtime/Analytics/TelemetryUtils/Private/TelemetryRouter.cpp
@@ -4,6 +4,15 @@
 #include "TelemetryUtils.h"
 #include "Misc/ScopeRWLock.h"
 
+#if DO_CHECK
+thread_local bool TelemetryRouterReentrancyGuard = false;
+#endif
+
+void FTelemetryRouter::CheckNotReentrant()
+{
+    check(TelemetryRouterReentrancyGuard == 0);
+}
+
 FTelemetryRouter::FTelemetryRouter()
 {
 }
@@ -19,8 +28,11 @@ FTelemetryRouter& FTelemetryRouter::Get()
 
 void FTelemetryRouter::ProvideTelemetryInternal(FGuid Key, FMemoryView Data) 
 {
+    CheckNotReentrant();
     FReadScopeLock Lock(SinkLock);
-    //TGuardValue Guard(ReentrancyGuard, FPlatformTLS::GetCurrentThreadId());
+#if DO_CHECK
+    TGuardValue Guard(TelemetryRouterReentrancyGuard, true);
+#endif
     FSinkSet* Sinks = KeyToSinks.Find(Key);
     if (Sinks)
     {
@@ -43,8 +55,11 @@ void FTelemetryRouter::ProvideTelemetryInternal(FGuid Key, FMemoryView Data)
 
 void FTelemetryRouter::RegisterTelemetrySinkInternal(FGuid Key, SIZE_T Size, FDelegateHandle InHandle, TFunction<bool(FMemoryView)> Sink)
 {
+    CheckNotReentrant();
     FWriteScopeLock Lock(SinkLock);
-    //TGuardValue Guard(ReentrancyGuard, FPlatformTLS::GetCurrentThreadId());
+#if DO_CHECK
+    TGuardValue Guard(TelemetryRouterReentrancyGuard, true);
+#endif
     FTelemetryRouter::FSinkSet* Sinks = KeyToSinks.Find(Key);
     if (!Sinks)
     {
@@ -60,8 +75,11 @@ void FTelemetryRouter::RegisterTelemetrySinkInternal(FGuid Key, SIZE_T Size, FDe
 
 void FTelemetryRouter::UnregisterTelemetrySinkInternal(FGuid Key, FDelegateHandle InHandle)
 {
+    CheckNotReentrant();
     FWriteScopeLock Lock(SinkLock);
-    //TGuardValue Guard(ReentrancyGuard, FPlatformTLS::GetCurrentThreadId());
+#if DO_CHECK
+    TGuardValue Guard(TelemetryRouterReentrancyGuard, true);
+#endif
     FTelemetryRouter::FSinkSet* Sinks = KeyToSinks.Find(Key);
     if (Sinks)
     {
diff --git a/Engine/Source/Runtime/Analytics/TelemetryUtils/Public/TelemetryRouter.h b/Engine/Source/Runtime/Analytics/TelemetryUtils/Public/TelemetryRouter.h
--- a/Engine/Source/Runtime/Analytics/TelemetryUtils/Public/TelemetryRouter.h
+++ b/Engine/Source/Runtime/Analytics/TelemetryUtils/Public/TelemetryRouter.h
@@ -82,7 +82,6 @@ public:
     template<typename DATA_TYPE>
     inline void ProvideTelemetry(const DATA_TYPE& Data)
     {
-        //check(ReentrancyGuard != FPlatformTLS::GetCurrentThreadId());
         ProvideTelemetryInternal(UE::Telemetry::Private::GetDataKey<DATA_TYPE>(), MakeMemoryView(&Data, sizeof(DATA_TYPE)));
     }
     
@@ -94,7 +93,6 @@ public:
     template<typename DATA_TYPE>
     inline FDelegateHandle OnTelemetry(TDelegate<void(const DATA_TYPE&)> Sink)
     {
-        //check(ReentrancyGuard != FPlatformTLS::GetCurrentThreadId());
         check(Sink.IsBound());
         FDelegateHandle Handle = Sink.GetHandle();
         RegisterTelemetrySinkInternal(UE::Telemetry::Private::GetDataKey<DATA_TYPE>(), sizeof(DATA_TYPE), Handle,
@@ -114,7 +112,6 @@ public:
     >
     inline FDelegateHandle OnTelemetry(CALLABLE&& Sink)
     {
-        //check(ReentrancyGuard != FPlatformTLS::GetCurrentThreadId());
         FDelegateHandle Handle{ FDelegateHandle::GenerateNewHandle };
         RegisterTelemetrySinkInternal(UE::Telemetry::Private::GetDataKey<DATA_TYPE>(), sizeof(DATA_TYPE), Handle,
             [Sink=MoveTemp(Sink)](FMemoryView Data) -> bool {
@@ -130,7 +127,6 @@ public:
     template<typename DATA_TYPE>
     inline void UnregisterTelemetrySink(FDelegateHandle Handle)
     {
-        //check(ReentrancyGuard != FPlatformTLS::GetCurrentThreadId());
         UnregisterTelemetrySinkInternal(UE::Telemetry::Private::GetDataKey<DATA_TYPE>(), Handle);
     }
     
@@ -141,6 +137,7 @@ private:
     TELEMETRYUTILS_API void RegisterTelemetrySinkInternal(FGuid Key, SIZE_T ExpectedSize, FDelegateHandle InHandle, TFunction<bool(FMemoryView)> Sink);
     TELEMETRYUTILS_API void UnregisterTelemetrySinkInternal(FGuid Key, FDelegateHandle InHandle);
     TELEMETRYUTILS_API void ProvideTelemetryInternal(FGuid Key, FMemoryView Data);
+    TELEMETRYUTILS_API void CheckNotReentrant();
     
     FRWLock SinkLock;
     
@@ -156,7 +153,4 @@ private:
     };
 
     TMap<FGuid, FSinkSet> KeyToSinks;
-	// Disabling ReentrancyGuard since it is not threadsafe, and a race condition can cause it to be set to Thread A again after Thread A raced
-	// Thread B and Thread A returns and then Thread B returns. Need to reconsider whether there is a way to cheaply prevent reentrancy.
-    //uint32 ReentrancyGuard = 0;
 };
\ No newline at end of file
