Fix for AnalyticsTracer shutdown to properly end the open spans

diff --git a/Engine/Source/Runtime/Analytics/Analytics/Private/AnalyticsTracer.cpp b/Engine/Source/Runtime/Analytics/Analytics/Private/AnalyticsTracer.cpp
--- a/Engine/Source/Runtime/Analytics/Analytics/Private/AnalyticsTracer.cpp
+++ b/Engine/Source/Runtime/Analytics/Analytics/Private/AnalyticsTracer.cpp
@@ -176,7 +176,7 @@ void FAnalyticsTracer::StartSession()
 void FAnalyticsTracer::EndSession()
 {
 	FScopeLock ScopeLock(&CriticalSection);	
-	SessionSpan->End();
+	EndSpan(SessionSpan);
 	SessionSpan.Reset();
 	AnalyticsProvider.Reset();
 }
