Replaced UE_NORETURN with [[noreturn]].
Replaced PLATFORM_COMPILER_HAS_DECLTYPE_AUTO blocks.
Replaced PLATFORM_COMPILER_HAS_FOLD_EXPRESSIONS blocks.

diff --git a/Engine/Source/Programs/ShaderCompileWorker/Private/ShaderCompileWorker.cpp b/Engine/Source/Programs/ShaderCompileWorker/Private/ShaderCompileWorker.cpp
--- a/Engine/Source/Programs/ShaderCompileWorker/Private/ShaderCompileWorker.cpp
+++ b/Engine/Source/Programs/ShaderCompileWorker/Private/ShaderCompileWorker.cpp
@@ -68,7 +68,7 @@ inline bool IsUsingUBA()
 }
 
 #if USING_CODE_ANALYSIS
-	UE_NORETURN static inline void ExitWithoutCrash(FSCWErrorCode::ECode ErrorCode, const FString& Message);
+	[[noreturn]] static inline void ExitWithoutCrash(FSCWErrorCode::ECode ErrorCode, const FString& Message);
 #endif
 
 static inline void ExitWithoutCrash(FSCWErrorCode::ECode ErrorCode, const FString& Message)
diff --git a/Engine/Source/Runtime/Analytics/Analytics/Public/AnalyticsConversion.h b/Engine/Source/Runtime/Analytics/Analytics/Public/AnalyticsConversion.h
--- a/Engine/Source/Runtime/Analytics/Analytics/Public/AnalyticsConversion.h
+++ b/Engine/Source/Runtime/Analytics/Analytics/Public/AnalyticsConversion.h
@@ -8,11 +8,7 @@
 
 /** Lexical conversion. Allow any type that we have a Lex for. Can't use universal references here because it then eats all non-perfect matches for the array and TMap conversions below, which we want to use a custom, analytics specific implementation for. */
 template <typename T>
-#if PLATFORM_COMPILER_HAS_DECLTYPE_AUTO
 inline decltype(auto) AnalyticsConversionToString(const T& Value)
-#else
-inline auto AnalyticsConversionToString(const T& Value) -> decltype(LexToString(Value)) 
-#endif
 {
 	return LexToString(Value);
 }
@@ -105,11 +101,7 @@ FString AnalyticsConversionToString(const TMap<KeyType, ValueType, Allocator, Ke
 namespace AnalyticsConversion
 {
 	template <typename T>
-#if PLATFORM_COMPILER_HAS_DECLTYPE_AUTO
 	inline decltype(auto) ToString(T&& Value)
-#else
-	inline auto ToString(T&& Value) -> decltype(AnalyticsConversionToString(Forward<T>(Value))) 
-#endif
 	{
 		return AnalyticsConversionToString(Forward<T>(Value));
 	}
diff --git a/Engine/Source/Runtime/Core/Public/Misc/TVariant.h b/Engine/Source/Runtime/Core/Public/Misc/TVariant.h
--- a/Engine/Source/Runtime/Core/Public/Misc/TVariant.h
+++ b/Engine/Source/Runtime/Core/Public/Misc/TVariant.h
@@ -200,12 +200,7 @@ template <
 >
 decltype(auto) Visit(Func&& Callable, Variants&&... Args)
 {
-#if PLATFORM_COMPILER_HAS_FOLD_EXPRESSIONS
 	constexpr SIZE_T NumPermutations = (1 * ... * (TVariantSize_V<std::decay_t<Variants>>));
-#else
-	constexpr SIZE_T VariantSizes[] = { TVariantSize_V<std::decay_t<Variants>>... };
-	constexpr SIZE_T NumPermutations = UE::Core::Private::Multiply(VariantSizes, sizeof...(Variants));
-#endif
 
 	return UE::Core::Private::VisitImpl(
 		UE::Core::Private::EncodeIndices(Args...),
diff --git a/Engine/Source/Runtime/Core/Public/Misc/TVariantMeta.h b/Engine/Source/Runtime/Core/Public/Misc/TVariantMeta.h
--- a/Engine/Source/Runtime/Core/Public/Misc/TVariantMeta.h
+++ b/Engine/Source/Runtime/Core/Public/Misc/TVariantMeta.h
@@ -325,20 +325,6 @@ namespace Private
 		return EncodedIndex % *VariantSizes;
 	}
 
-#if !PLATFORM_COMPILER_HAS_FOLD_EXPRESSIONS
-	/** Used to determine the total number of possible Visit invocations when fold expressions are not available. */
-	constexpr SIZE_T Multiply(const SIZE_T* Args, SIZE_T Num)
-	{
-		SIZE_T Result = 1;
-		while (Num)
-		{
-			Result *= *Args++;
-			--Num;
-		}
-		return Result;
-	}
-#endif
-
 	/** Cast a TVariant to its private base */
 	template <typename... Ts>
 	FORCEINLINE TVariantStorage<Ts...>& CastToStorage(TVariant<Ts...>& Variant)
diff --git a/Engine/Source/Runtime/CoreUObject/Public/Templates/Casts.h b/Engine/Source/Runtime/CoreUObject/Public/Templates/Casts.h
--- a/Engine/Source/Runtime/CoreUObject/Public/Templates/Casts.h
+++ b/Engine/Source/Runtime/CoreUObject/Public/Templates/Casts.h
@@ -36,7 +36,7 @@ template<class TClass> class TSubclassOf;
 template <typename Type> struct TCastFlags;
 /// @endcond
 
-UE_NORETURN COREUOBJECT_API void CastLogError(const TCHAR* FromType, const TCHAR* ToType);
+[[noreturn]] COREUOBJECT_API void CastLogError(const TCHAR* FromType, const TCHAR* ToType);
 
 /**
  * Metafunction which detects whether or not a class is an IInterface.  Rules:
