Address various suggestions for improvement from previous code review.
Simplified the parent child relationship. Moved this code out of AnalyticsSpan and int AnalyticsTracer. Resulting in much cleaner code and faster search.
No longer supporting the Span_ParentName in the event telemetry as Spans no longer know anythign about their parent child relationship. This behavior is specfic to the IAnalyticsTracer implementation.
Added GetId() to AnalyticsSpan. Tracking of parent/child relationships by ID much more efficient now.
Removed StartSpan( Span, Attribuites ) from the public API as this is not used.
Fixed issue where parent span attrbiutes are added twice
Fixed typo in EndPIE callabck

diff --git a/Engine/Plugins/Experimental/StudioTelemetry/Source/StudioTelemetry/Private/StudioTelemetry.cpp b/Engine/Plugins/Experimental/StudioTelemetry/Source/StudioTelemetry/Private/StudioTelemetry.cpp
--- a/Engine/Plugins/Experimental/StudioTelemetry/Source/StudioTelemetry/Private/StudioTelemetry.cpp
+++ b/Engine/Plugins/Experimental/StudioTelemetry/Source/StudioTelemetry/Private/StudioTelemetry.cpp
@@ -238,11 +238,6 @@ TSharedPtr<IAnalyticsSpan> FStudioTelemetry::StartSpan(const FName Name, TShared
 	return AnalyticsTracer.IsValid() ? AnalyticsTracer->StartSpan(Name, ParentSpan, AdditionalAttributes)  : TSharedPtr<IAnalyticsSpan>();
 }
 
-bool FStudioTelemetry::StartSpan(TSharedPtr<IAnalyticsSpan> Span, const TArray<FAnalyticsEventAttribute>& AdditionalAttributes)
-{
-	return AnalyticsTracer.IsValid() ? AnalyticsTracer->StartSpan(Span, AdditionalAttributes) : false;
-}
-
 bool FStudioTelemetry::EndSpan(TSharedPtr<IAnalyticsSpan> Span, const TArray<FAnalyticsEventAttribute>& AdditionalAttributes)
 {
 	return AnalyticsTracer.IsValid() ? AnalyticsTracer->EndSpan(Span, AdditionalAttributes) : false;
diff --git a/Engine/Plugins/Experimental/StudioTelemetry/Source/StudioTelemetry/Private/StudioTelemetryEditor.cpp b/Engine/Plugins/Experimental/StudioTelemetry/Source/StudioTelemetry/Private/StudioTelemetryEditor.cpp
--- a/Engine/Plugins/Experimental/StudioTelemetry/Source/StudioTelemetry/Private/StudioTelemetryEditor.cpp
+++ b/Engine/Plugins/Experimental/StudioTelemetry/Source/StudioTelemetry/Private/StudioTelemetryEditor.cpp
@@ -707,7 +707,7 @@ void FStudioTelemetryEditor::Initialize()
 		{
 			if (PIESpan.IsValid())
 			{
-				// PIE his ending so no longer interactive
+				// PIE is ending so no longer interactive
 				FStudioTelemetry::Get().EndSpan(PIEInteractSpan);
 				PIEShutdownSpan = FStudioTelemetry::Get().StartSpan(PIEShutdownSpanName, PIESpan);
 			}
diff --git a/Engine/Plugins/Experimental/StudioTelemetry/Source/StudioTelemetry/Public/StudioTelemetry.h b/Engine/Plugins/Experimental/StudioTelemetry/Source/StudioTelemetry/Public/StudioTelemetry.h
--- a/Engine/Plugins/Experimental/StudioTelemetry/Source/StudioTelemetry/Public/StudioTelemetry.h
+++ b/Engine/Plugins/Experimental/StudioTelemetry/Source/StudioTelemetry/Public/StudioTelemetry.h
@@ -58,9 +58,6 @@ public:
 	/** Start a new span specifying the parent*/
 	STUDIOTELEMETRY_API TSharedPtr<IAnalyticsSpan> StartSpan(const FName Name, TSharedPtr<IAnalyticsSpan> ParentSpan, const TArray<FAnalyticsEventAttribute>& AdditionalAttributes = {});
 
-	/** Start an existing span*/
-	STUDIOTELEMETRY_API bool StartSpan(TSharedPtr<IAnalyticsSpan> Span, const TArray<FAnalyticsEventAttribute>& AdditionalAttributes = {});
-
 	/** End an existing span*/
 	STUDIOTELEMETRY_API bool EndSpan(TSharedPtr<IAnalyticsSpan> Span, const TArray<FAnalyticsEventAttribute>& AdditionalAttributes = {});
 
diff --git a/Engine/Source/Runtime/Analytics/Analytics/Private/AnalyticsTracer.cpp b/Engine/Source/Runtime/Analytics/Analytics/Private/AnalyticsTracer.cpp
--- a/Engine/Source/Runtime/Analytics/Analytics/Private/AnalyticsTracer.cpp
+++ b/Engine/Source/Runtime/Analytics/Analytics/Private/AnalyticsTracer.cpp
@@ -47,11 +47,6 @@ double FAnalyticsSpan::GetDuration() const
 	return Duration;
 }
 
-void FAnalyticsSpan::SetParentSpan(TSharedPtr<IAnalyticsSpan> NewSpanParent)
-{
-	ParentSpan = NewSpanParent;
-}
-
 void FAnalyticsSpan::Start(const TArray<FAnalyticsEventAttribute>& AdditionalAttributes)
 {
 	// Create a new Guid for this flow, can we assume it is unique?
@@ -85,16 +80,9 @@ void FAnalyticsSpan::End(const TArray<FAnalyticsEventAttribute>& AdditionalAttri
 		
 	TRACE_END_REGION(*Name.ToString());
 
-	// Append the parent attributes and the the additional attributes to the current span attributes, these will get passed down to the child spans
+	// Add attributes and the the additional attributes to the current span attributes, these will get passed down to the child spans
 	AddAttributes(AdditionalAttributes);
 
-	TSharedPtr< IAnalyticsSpan> ParentSpanShared = ParentSpan.Pin();
-
-	if (ParentSpanShared.IsValid())
-	{
-		AddAttributes(ParentSpanShared->GetAttributes());
-	}
-
 	const uint32 SpanSchemaVersion = 1;
 	const FString SpanEventName = TEXT("Span");
 
@@ -103,7 +91,6 @@ void FAnalyticsSpan::End(const TArray<FAnalyticsEventAttribute>& AdditionalAttri
 	EventAttributes.Add(FAnalyticsEventAttribute(TEXT("SchemaVersion"), SpanSchemaVersion));
 	EventAttributes.Add(FAnalyticsEventAttribute(TEXT("Span_Name"), Name.ToString()));
 	EventAttributes.Add(FAnalyticsEventAttribute(TEXT("Span_GUID"), Guid.ToString()));
-	EventAttributes.Add(FAnalyticsEventAttribute(TEXT("Span_ParentName"), ParentSpanShared.IsValid() ? ParentSpanShared->GetName().ToString() : TEXT("")));
 	EventAttributes.Add(FAnalyticsEventAttribute(TEXT("Span_ThreadId"), ThreadId));
 	EventAttributes.Add(FAnalyticsEventAttribute(TEXT("Span_Depth"), StackDepth));
 	EventAttributes.Add(FAnalyticsEventAttribute(TEXT("Span_StartUTC"), StartTime.ToUnixTimestampDecimal()));
@@ -138,6 +125,11 @@ const FName& FAnalyticsSpan::GetName() const
 	return Name;
 }
 
+FGuid FAnalyticsSpan::GetId() const
+{
+	return Guid;
+}
+
 const TArray<FAnalyticsEventAttribute>& FAnalyticsSpan::GetAttributes() const
 {
 	return Attributes;
@@ -148,11 +140,6 @@ uint32 FAnalyticsSpan::GetStackDepth() const
 	return StackDepth;
 }
 
-TSharedPtr<IAnalyticsSpan> FAnalyticsSpan::GetParentSpan() const
-{
-	return ParentSpan.Pin();
-}
-
 void FAnalyticsTracer::SetProvider(TSharedPtr<IAnalyticsProvider> InProvider)
 {
 	AnalyticsProvider = InProvider;
@@ -175,19 +162,10 @@ void FAnalyticsTracer::EndSession()
 	EndSpan(SessionSpan);
 	SessionSpan.Reset();
 
-	// Stop any active spans, go from stack bottom first so parent spans will end their children 
+	// Stop any active spans, go from stack bottom first so parent spans will end their children
 	while (ActiveSpanStack.Num())
 	{	
-		TSharedPtr<IAnalyticsSpan> Span = ActiveSpanStack[0];
-
-		if (Span.IsValid())
-		{
-			EndSpan(Span);
-		}	
-		else
-		{
-			ActiveSpanStack.Remove(Span);
-		}
+		EndSpan(ActiveSpanStack[0]);
 	}
 
 	ActiveSpanStack.Reset();
@@ -199,7 +177,22 @@ TSharedPtr<IAnalyticsSpan> FAnalyticsTracer::StartSpan(const FName NewSpanName,
 {
 	FScopeLock ScopeLock(&CriticalSection);
 	TSharedPtr<IAnalyticsSpan> NewSpan = MakeShared<FAnalyticsSpan>(NewSpanName);
-	NewSpan->SetParentSpan(ParentSpan);
+
+	// Add the child to the parent's list of child spans
+	if (ParentSpan.IsValid())
+	{
+		if (SpanHeirarchy.Find(ParentSpan->GetId()))
+		{
+			// Child list exists for this span already so append the new child for this span
+			SpanHeirarchy[ParentSpan->GetId()].Emplace(NewSpan);
+		}
+		else
+		{
+			// Create a new child list for this span
+			SpanHeirarchy.Emplace(ParentSpan->GetId(), { NewSpan });
+		}
+	}
+	
 
 	return StartSpanInternal(NewSpan, AdditionalAttributes) ? NewSpan : TSharedPtr<IAnalyticsSpan>();
 }
@@ -241,28 +234,25 @@ bool FAnalyticsTracer::EndSpanInternal(TSharedPtr<IAnalyticsSpan> Span, const TA
 
 		ActiveSpanStack.Remove(Span);
 
-		// Remove any active children of this span, inherit on the parent's attributes in the child
-		bool RemovedActiveChild = false;
-
-		do 
+		// End any children of this span
+		const TArray<TWeakPtr<IAnalyticsSpan>>* ChildSpans = SpanHeirarchy.Find(Span->GetId());
+		
+		if (ChildSpans != nullptr)
 		{
-			RemovedActiveChild = false;
-
-			for (int32 i = 0; i < ActiveSpanStack.Num(); ++i)
+			for (TWeakPtr<IAnalyticsSpan> ChildSpanWeakPtr : *ChildSpans)
 			{
-				TSharedPtr<IAnalyticsSpan> ChildSpan = ActiveSpanStack[i];
+				TSharedPtr<IAnalyticsSpan> ChildSpan = ChildSpanWeakPtr.Pin();
 
-				if (ChildSpan.IsValid() && ChildSpan->GetParentSpan() == Span)
+				if (ChildSpan.IsValid())
 				{
+					// Pass the parent's attributes to the children as it ends
 					EndSpanInternal(ChildSpan, Span->GetAttributes());
-					RemovedActiveChild = true;
-					break;
 				}
 			}
+		}
 
-		} while (RemovedActiveChild==true);
-
-		return true;
+		// Remove this span's child list 
+		SpanHeirarchy.Remove(Span->GetId());
 	}
 
 	return false;
diff --git a/Engine/Source/Runtime/Analytics/Analytics/Public/AnalyticsTracer.h b/Engine/Source/Runtime/Analytics/Analytics/Public/AnalyticsTracer.h
--- a/Engine/Source/Runtime/Analytics/Analytics/Public/AnalyticsTracer.h
+++ b/Engine/Source/Runtime/Analytics/Analytics/Public/AnalyticsTracer.h
@@ -26,10 +26,9 @@ public:
 	virtual const TArray<FAnalyticsEventAttribute>& GetAttributes() const override;
 	virtual void SetStackDepth(uint32 Depth) override;
 	virtual uint32 GetStackDepth() const override;
-	virtual void SetParentSpan(TSharedPtr<IAnalyticsSpan> ParentSpan) override;
-	virtual TSharedPtr<IAnalyticsSpan> GetParentSpan() const override;
 	virtual double GetDuration() const override;
 	virtual bool GetIsActive() const override;
+	virtual FGuid GetId() const override;
 
 private:
 
@@ -42,8 +41,7 @@ private:
 	double								Duration = 0;
 	bool								IsActive = false;
 	TSharedPtr<IAnalyticsProvider>		AnalyticsProvider;
-	TArray<FAnalyticsEventAttribute>	Attributes;
-	TWeakPtr<IAnalyticsSpan>			ParentSpan;
+	TArray<FAnalyticsEventAttribute>	Attributes;;
 };
 
 /**
@@ -73,10 +71,11 @@ private:
 	bool EndSpanInternal(TSharedPtr<IAnalyticsSpan> Span, const TArray<FAnalyticsEventAttribute>& AdditionalAttributes);
 	TSharedPtr<IAnalyticsSpan> GetSpanInternal(const FName Name);
 
-	TSharedPtr<IAnalyticsProvider>			AnalyticsProvider;	// The Analytics provider we will send our span events to
-	TSharedPtr<IAnalyticsSpan>				SessionSpan;		// The root span, this will always be present in an active session
-	TArray<TSharedPtr<IAnalyticsSpan>>		ActiveSpanStack;	// Stack of active spans as WeakPtrs
-	FCriticalSection						CriticalSection;	
+	TSharedPtr<IAnalyticsProvider>					AnalyticsProvider;	// The Analytics provider we will send our span events to
+	TSharedPtr<IAnalyticsSpan>						SessionSpan;		// The root span, this will always be present in an active session
+	TArray<TSharedPtr<IAnalyticsSpan>>				ActiveSpanStack;	// Stack of active spans as WeakPtrs
+	TMap<FGuid, TArray<TWeakPtr<IAnalyticsSpan>>>	SpanHeirarchy;		// List of child spans per parent span
+	FCriticalSection								CriticalSection;	
 };
 
 
diff --git a/Engine/Source/Runtime/Analytics/Analytics/Public/Interfaces/IAnalyticsTracer.h b/Engine/Source/Runtime/Analytics/Analytics/Public/Interfaces/IAnalyticsTracer.h
--- a/Engine/Source/Runtime/Analytics/Analytics/Public/Interfaces/IAnalyticsTracer.h
+++ b/Engine/Source/Runtime/Analytics/Analytics/Public/Interfaces/IAnalyticsTracer.h
@@ -57,11 +57,8 @@ public:
 	// Is the span active?
 	virtual bool GetIsActive() const = 0;
 
-	/** Set the parent span */
-	virtual void SetParentSpan(TSharedPtr<IAnalyticsSpan> ParentSpan) = 0;
-
-	/** Get the parent span */
-	virtual TSharedPtr<IAnalyticsSpan> GetParentSpan() const = 0;
+	// Get unique Id?
+	virtual FGuid GetId()  const = 0;
 };
 
 /**
