Various tweaks and improvements to Tracer API suggested by Wes

diff --git a/Engine/Source/Runtime/Analytics/Analytics/Private/AnalyticsTracer.cpp b/Engine/Source/Runtime/Analytics/Analytics/Private/AnalyticsTracer.cpp
--- a/Engine/Source/Runtime/Analytics/Analytics/Private/AnalyticsTracer.cpp
+++ b/Engine/Source/Runtime/Analytics/Analytics/Private/AnalyticsTracer.cpp
@@ -197,16 +197,6 @@ TSharedPtr<IAnalyticsSpan> FAnalyticsTracer::StartSpan(const FName NewSpanName,
 	return StartSpanInternal(NewSpan, AdditionalAttributes) ? NewSpan : TSharedPtr<IAnalyticsSpan>();
 }
 
-bool FAnalyticsTracer::StartSpan(TSharedPtr<IAnalyticsSpan> Span, const TArray<FAnalyticsEventAttribute>& AdditionalAttributes)
-{
-	if (Span.IsValid() && Span->GetIsActive() == false)
-	{
-		return StartSpanInternal(Span, AdditionalAttributes);
-	}
-
-	return false;
-}
-
 bool FAnalyticsTracer::StartSpanInternal(TSharedPtr<IAnalyticsSpan> Span, const TArray<FAnalyticsEventAttribute>& AdditionalAttributes)
 {
 	TSharedPtr<IAnalyticsSpan> LastAdddedActiveSpan = ActiveSpanStack.Num()? ActiveSpanStack.Top(): TSharedPtr<IAnalyticsSpan>();
@@ -241,13 +231,8 @@ bool FAnalyticsTracer::EndSpanInternal(TSharedPtr<IAnalyticsSpan> Span, const TA
 		{
 			for (TWeakPtr<IAnalyticsSpan> ChildSpanWeakPtr : *ChildSpans)
 			{
-				TSharedPtr<IAnalyticsSpan> ChildSpan = ChildSpanWeakPtr.Pin();
-
-				if (ChildSpan.IsValid())
-				{
-					// Pass the parent's attributes to the children as it ends
-					EndSpanInternal(ChildSpan, Span->GetAttributes());
-				}
+				// Pass the parent's attributes to the children as it ends
+				EndSpanInternal(ChildSpanWeakPtr.Pin(), Span->GetAttributes());
 			}
 		}
 
diff --git a/Engine/Source/Runtime/Analytics/Analytics/Public/AnalyticsTracer.h b/Engine/Source/Runtime/Analytics/Analytics/Public/AnalyticsTracer.h
--- a/Engine/Source/Runtime/Analytics/Analytics/Public/AnalyticsTracer.h
+++ b/Engine/Source/Runtime/Analytics/Analytics/Public/AnalyticsTracer.h
@@ -59,7 +59,6 @@ public:
 	virtual void EndSession() override;
 	virtual void SetProvider(TSharedPtr<IAnalyticsProvider> AnalyticsProvider) override;
 	virtual TSharedPtr<IAnalyticsSpan> StartSpan(const FName Name, TSharedPtr<IAnalyticsSpan> ParentSpan, const TArray<FAnalyticsEventAttribute>& AdditionalAttributes = {}) override;
-	virtual bool StartSpan(TSharedPtr<IAnalyticsSpan> Span, const TArray<FAnalyticsEventAttribute>& AdditionalAttributes = {}) override;
 	virtual bool EndSpan(TSharedPtr<IAnalyticsSpan> Span, const TArray<FAnalyticsEventAttribute>& AdditionalAttributes = {}) override;
 	virtual TSharedPtr<IAnalyticsSpan> GetCurrentSpan() const override;
 	virtual TSharedPtr<IAnalyticsSpan> GetSessionSpan() const override;
diff --git a/Engine/Source/Runtime/Analytics/Analytics/Public/Interfaces/IAnalyticsTracer.h b/Engine/Source/Runtime/Analytics/Analytics/Public/Interfaces/IAnalyticsTracer.h
--- a/Engine/Source/Runtime/Analytics/Analytics/Public/Interfaces/IAnalyticsTracer.h
+++ b/Engine/Source/Runtime/Analytics/Analytics/Public/Interfaces/IAnalyticsTracer.h
@@ -88,9 +88,6 @@ public:
 	/** Start a new span specifying an optional parent. EndSpan is called recursively on children. Parent attributes are passed onto children  */
 	virtual TSharedPtr<IAnalyticsSpan> StartSpan(const FName Name, TSharedPtr<IAnalyticsSpan> ParentSpan=TSharedPtr<IAnalyticsSpan>(), const TArray<FAnalyticsEventAttribute>& AdditionalAttributes = {}) = 0;
 
-	/** Start a an existing span*/
-	virtual bool StartSpan(TSharedPtr<IAnalyticsSpan> Span, const TArray<FAnalyticsEventAttribute>& AdditionalAttributes = {}) = 0;
-
 	/** End an existing span*/
 	virtual bool EndSpan(TSharedPtr<IAnalyticsSpan>, const TArray<FAnalyticsEventAttribute>& AdditionalAttributes = {}) = 0;
 
