Move the initial declaration of ::BlockUntilFlushed from IAnalyticsProviderET to it's parent class IAnalyticsProvider and improve the documentation showing the difference between ::FlushEvents and ::BlockUntilFlushed.

#rnx

diff --git a/Engine/Plugins/Experimental/StudioTelemetry/Source/AnalyticsLog/Private/AnalyticsProviderLog.cpp b/Engine/Plugins/Experimental/StudioTelemetry/Source/AnalyticsLog/Private/AnalyticsProviderLog.cpp
--- a/Engine/Plugins/Experimental/StudioTelemetry/Source/AnalyticsLog/Private/AnalyticsProviderLog.cpp
+++ b/Engine/Plugins/Experimental/StudioTelemetry/Source/AnalyticsLog/Private/AnalyticsProviderLog.cpp
@@ -58,6 +58,10 @@ void FAnalyticsProviderLog::FlushEvents()
 {
 }
 
+void FAnalyticsProviderLog::BlockUntilFlushed(float InTimeoutSec)
+{
+}
+
 void FAnalyticsProviderLog::SetDefaultEventAttributes(TArray<FAnalyticsEventAttribute>&& Attributes)
 {
 	DefaultEventAttributes = Attributes;
diff --git a/Engine/Plugins/Experimental/StudioTelemetry/Source/AnalyticsLog/Private/AnalyticsProviderLog.h b/Engine/Plugins/Experimental/StudioTelemetry/Source/AnalyticsLog/Private/AnalyticsProviderLog.h
--- a/Engine/Plugins/Experimental/StudioTelemetry/Source/AnalyticsLog/Private/AnalyticsProviderLog.h
+++ b/Engine/Plugins/Experimental/StudioTelemetry/Source/AnalyticsLog/Private/AnalyticsProviderLog.h
@@ -22,6 +22,7 @@ public:
 	virtual void EndSession() override;
 
 	virtual void FlushEvents() override;
+	virtual void BlockUntilFlushed(float InTimeoutSec) override;
 
 	virtual void SetDefaultEventAttributes(TArray<FAnalyticsEventAttribute>&& Attributes)  override;
 	virtual TArray<FAnalyticsEventAttribute> GetDefaultEventAttributesSafe() const override;
diff --git a/Engine/Plugins/Experimental/StudioTelemetry/Source/StudioTelemetry/Private/AnalyticsProviderMulticast.cpp b/Engine/Plugins/Experimental/StudioTelemetry/Source/StudioTelemetry/Private/AnalyticsProviderMulticast.cpp
--- a/Engine/Plugins/Experimental/StudioTelemetry/Source/StudioTelemetry/Private/AnalyticsProviderMulticast.cpp
+++ b/Engine/Plugins/Experimental/StudioTelemetry/Source/StudioTelemetry/Private/AnalyticsProviderMulticast.cpp
@@ -205,6 +205,14 @@ void FAnalyticsProviderMulticast::FlushEvents()
 	}
 }
 
+void FAnalyticsProviderMulticast::BlockUntilFlushed(float InTimeoutSec)
+{
+	for (TProviders::TConstIterator it(Providers); it; ++it)
+	{
+		(*it).Value->BlockUntilFlushed(InTimeoutSec);
+	}
+}
+
 void FAnalyticsProviderMulticast::SetRecordEventCallback(OnRecordEvent Callback)
 {
 	RecordEventCallback = Callback;
@@ -213,7 +221,7 @@ void FAnalyticsProviderMulticast::SetRecordEventCallback(OnRecordEvent Callback)
 void FAnalyticsProviderMulticast::RecordEvent(const FString& EventName, const TArray<FAnalyticsEventAttribute>& Attributes)
 {
 #if UE_BUILD_DEBUG || UE_BUILD_DEVELOPMENT
-	// Expose events that have duplicate aatibute names. This is is not handled by the analytics backends in any reliable way.
+	// Expose events that have duplicate attribute names. This is is not handled by the analytics backends in any reliable way.
 	for (int32 index0 = 0; index0 < Attributes.Num(); ++index0)
 	{
 		for (int32 index1 = index0 + 1; index1 < Attributes.Num(); ++index1)
diff --git a/Engine/Plugins/Experimental/StudioTelemetry/Source/StudioTelemetry/Private/AnalyticsProviderMulticast.h b/Engine/Plugins/Experimental/StudioTelemetry/Source/StudioTelemetry/Private/AnalyticsProviderMulticast.h
--- a/Engine/Plugins/Experimental/StudioTelemetry/Source/StudioTelemetry/Private/AnalyticsProviderMulticast.h
+++ b/Engine/Plugins/Experimental/StudioTelemetry/Source/StudioTelemetry/Private/AnalyticsProviderMulticast.h
@@ -37,7 +37,8 @@ public:
 	virtual bool StartSession(const TArray<FAnalyticsEventAttribute>& Attributes = {}) override;
 	virtual void EndSession() override;
 	virtual void FlushEvents() override;
-	
+	virtual void BlockUntilFlushed(float InTimeoutSec) override;
+
 	virtual void SetDefaultEventAttributes(TArray<FAnalyticsEventAttribute>&& Attributes)  override;
 	virtual TArray<FAnalyticsEventAttribute> GetDefaultEventAttributesSafe() const override;
 	virtual int32 GetDefaultEventAttributeCount() const  override;
diff --git a/Engine/Plugins/Runtime/Analytics/Adjust/Source/AndroidAdjust/Private/AndroidAdjust.cpp b/Engine/Plugins/Runtime/Analytics/Adjust/Source/AndroidAdjust/Private/AndroidAdjust.cpp
--- a/Engine/Plugins/Runtime/Analytics/Adjust/Source/AndroidAdjust/Private/AndroidAdjust.cpp
+++ b/Engine/Plugins/Runtime/Analytics/Adjust/Source/AndroidAdjust/Private/AndroidAdjust.cpp
@@ -260,6 +260,15 @@ void FAnalyticsProviderAdjust::EndSession()
 #endif
 }
 
+void FAnalyticsProviderAdjust::BlockUntilFlushed(float InTimeoutSec)
+{
+#if WITH_ADJUST
+	UE_LOG(LogAnalytics, Display, TEXT("AndroidAdjust::BlockUntilFlushed"));
+#else
+	UE_LOG(LogAnalytics, Warning, TEXT("WITH_ADJUST=0. Are you missing the SDK?"));
+#endif
+}
+
 void FAnalyticsProviderAdjust::FlushEvents()
 {
 #if WITH_ADJUST
diff --git a/Engine/Plugins/Runtime/Analytics/Adjust/Source/AndroidAdjust/Private/AndroidAdjustProvider.h b/Engine/Plugins/Runtime/Analytics/Adjust/Source/AndroidAdjust/Private/AndroidAdjustProvider.h
--- a/Engine/Plugins/Runtime/Analytics/Adjust/Source/AndroidAdjust/Private/AndroidAdjustProvider.h
+++ b/Engine/Plugins/Runtime/Analytics/Adjust/Source/AndroidAdjust/Private/AndroidAdjustProvider.h
@@ -39,6 +39,7 @@ public:
 	virtual bool StartSession(const TArray<FAnalyticsEventAttribute>& Attributes) override;
 	virtual void EndSession() override;
 	virtual void FlushEvents() override;
+	virtual void BlockUntilFlushed(float InTimeoutSec) override;
 
 	virtual void SetUserID(const FString& InUserID) override;
 	virtual FString GetUserID() const override;
diff --git a/Engine/Plugins/Runtime/Analytics/Adjust/Source/IOSAdjust/Private/IOSAdjust.cpp b/Engine/Plugins/Runtime/Analytics/Adjust/Source/IOSAdjust/Private/IOSAdjust.cpp
--- a/Engine/Plugins/Runtime/Analytics/Adjust/Source/IOSAdjust/Private/IOSAdjust.cpp
+++ b/Engine/Plugins/Runtime/Analytics/Adjust/Source/IOSAdjust/Private/IOSAdjust.cpp
@@ -211,6 +211,16 @@ void FAnalyticsProviderAdjust::FlushEvents()
 #endif
 }
 
+void FAnalyticsProviderAdjust::BlockUntilFlushed(float InTimeoutSec)
+{
+#if WITH_ADJUST
+	[Adjust sendFirstPackages];
+	UE_LOG(LogAnalytics, Display, TEXT("IOSAdjust::BlockUntilFlushed"));
+#else
+	UE_LOG(LogAnalytics, Warning, TEXT("WITH_ADJUST=0. Are you missing the SDK?"));
+#endif
+}
+
 void FAnalyticsProviderAdjust::SetUserID(const FString& InUserID)
 {
 #if WITH_ADJUST
diff --git a/Engine/Plugins/Runtime/Analytics/Adjust/Source/IOSAdjust/Private/IOSAdjustProvider.h b/Engine/Plugins/Runtime/Analytics/Adjust/Source/IOSAdjust/Private/IOSAdjustProvider.h
--- a/Engine/Plugins/Runtime/Analytics/Adjust/Source/IOSAdjust/Private/IOSAdjustProvider.h
+++ b/Engine/Plugins/Runtime/Analytics/Adjust/Source/IOSAdjust/Private/IOSAdjustProvider.h
@@ -39,6 +39,7 @@ public:
 	virtual bool StartSession(const TArray<FAnalyticsEventAttribute>& Attributes) override;
 	virtual void EndSession() override;
 	virtual void FlushEvents() override;
+	virtual void BlockUntilFlushed(float InTimeoutSec) override;
 
 	virtual void SetUserID(const FString& InUserID) override;
 	virtual FString GetUserID() const override;
diff --git a/Engine/Plugins/Runtime/Analytics/AnalyticsMulticast/Source/AnalyticsMulticast/Private/AnalyticsMulticast.cpp b/Engine/Plugins/Runtime/Analytics/AnalyticsMulticast/Source/AnalyticsMulticast/Private/AnalyticsMulticast.cpp
--- a/Engine/Plugins/Runtime/Analytics/AnalyticsMulticast/Source/AnalyticsMulticast/Private/AnalyticsMulticast.cpp
+++ b/Engine/Plugins/Runtime/Analytics/AnalyticsMulticast/Source/AnalyticsMulticast/Private/AnalyticsMulticast.cpp
@@ -30,6 +30,7 @@ public:
 	virtual bool StartSession(const TArray<FAnalyticsEventAttribute>& Attributes) override;
 	virtual void EndSession() override;
 	virtual void FlushEvents() override;
+	virtual void BlockUntilFlushed(float InTimeoutSec) override;
 
 	virtual void SetUserID(const FString& InUserID) override;
 	virtual FString GetUserID() const override;
@@ -159,6 +160,14 @@ void FAnalyticsProviderMulticast::FlushEvents()
 	}
 }
 
+void FAnalyticsProviderMulticast::BlockUntilFlushed(float InTimeoutSec)
+{
+	for (TArray<TSharedPtr<IAnalyticsProvider> >::TConstIterator it(Providers); it; ++it)
+	{
+		(*it)->BlockUntilFlushed(InTimeoutSec);
+	}
+}
+
 void FAnalyticsProviderMulticast::SetUserID( const FString& InUserID )
 {
 	for (TArray<TSharedPtr<IAnalyticsProvider> >::TConstIterator it(Providers);it;++it)
diff --git a/Engine/Plugins/Runtime/Analytics/FileLogging/Source/FileLogging/Private/FileLogging.cpp b/Engine/Plugins/Runtime/Analytics/FileLogging/Source/FileLogging/Private/FileLogging.cpp
--- a/Engine/Plugins/Runtime/Analytics/FileLogging/Source/FileLogging/Private/FileLogging.cpp
+++ b/Engine/Plugins/Runtime/Analytics/FileLogging/Source/FileLogging/Private/FileLogging.cpp
@@ -115,6 +115,13 @@ void FAnalyticsProviderFileLogging::FlushEvents()
 	}
 }
 
+void FAnalyticsProviderFileLogging::BlockUntilFlushed(float InTimeoutSec)
+{
+	// We don't really need to block when flushing in this provider so we
+	// can just call flush.
+	FlushEvents();
+}
+
 void FAnalyticsProviderFileLogging::SetUserID(const FString& InUserID)
 {
 	if (!bHasSessionStarted)
diff --git a/Engine/Plugins/Runtime/Analytics/FileLogging/Source/FileLogging/Private/FileLoggingProvider.h b/Engine/Plugins/Runtime/Analytics/FileLogging/Source/FileLogging/Private/FileLoggingProvider.h
--- a/Engine/Plugins/Runtime/Analytics/FileLogging/Source/FileLogging/Private/FileLoggingProvider.h
+++ b/Engine/Plugins/Runtime/Analytics/FileLogging/Source/FileLogging/Private/FileLoggingProvider.h
@@ -39,6 +39,7 @@ public:
 	virtual bool StartSession(const TArray<FAnalyticsEventAttribute>& Attributes) override;
 	virtual void EndSession() override;
 	virtual void FlushEvents() override;
+	virtual void BlockUntilFlushed(float InTimeoutSec) override;
 
 	virtual void SetUserID(const FString& InUserID) override;
 	virtual FString GetUserID() const override;
diff --git a/Engine/Plugins/Runtime/Analytics/Flurry/Source/IOSFlurry/Source/IOSFlurry/Private/IOSFlurry.cpp b/Engine/Plugins/Runtime/Analytics/Flurry/Source/IOSFlurry/Source/IOSFlurry/Private/IOSFlurry.cpp
--- a/Engine/Plugins/Runtime/Analytics/Flurry/Source/IOSFlurry/Source/IOSFlurry/Private/IOSFlurry.cpp
+++ b/Engine/Plugins/Runtime/Analytics/Flurry/Source/IOSFlurry/Source/IOSFlurry/Private/IOSFlurry.cpp
@@ -102,6 +102,16 @@ void FAnalyticsProviderFlurry::FlushEvents()
 #endif
 }
 
+void FAnalyticsProviderFlurry::BlockUntilFlushed(float InTimeoutSec)
+{
+#if WITH_FLURRY
+	// Flurry doesn't support flushing a session
+	UE_LOG(LogAnalytics, Display, TEXT("IOSFlurry::BlockUntilFlushed - ignoring call"));
+#else
+	UE_LOG(LogAnalytics, Warning, TEXT("WITH_FLURRY=0. Are you missing the SDK?"));
+#endif
+}
+
 void FAnalyticsProviderFlurry::SetUserID(const FString& InUserID)
 {
 #if WITH_FLURRY
diff --git a/Engine/Plugins/Runtime/Analytics/Flurry/Source/IOSFlurry/Source/IOSFlurry/Private/IOSFlurryProvider.h b/Engine/Plugins/Runtime/Analytics/Flurry/Source/IOSFlurry/Source/IOSFlurry/Private/IOSFlurryProvider.h
--- a/Engine/Plugins/Runtime/Analytics/Flurry/Source/IOSFlurry/Source/IOSFlurry/Private/IOSFlurryProvider.h
+++ b/Engine/Plugins/Runtime/Analytics/Flurry/Source/IOSFlurry/Source/IOSFlurry/Private/IOSFlurryProvider.h
@@ -37,6 +37,7 @@ public:
 	virtual bool StartSession(const TArray<FAnalyticsEventAttribute>& Attributes) override;
 	virtual void EndSession() override;
 	virtual void FlushEvents() override;
+	virtual void BlockUntilFlushed(float InTimeoutSec) override;
 
 	virtual void SetUserID(const FString& InUserID) override;
 	virtual FString GetUserID() const override;
diff --git a/Engine/Source/Runtime/Analytics/Analytics/Public/Interfaces/IAnalyticsProvider.h b/Engine/Source/Runtime/Analytics/Analytics/Public/Interfaces/IAnalyticsProvider.h
--- a/Engine/Source/Runtime/Analytics/Analytics/Public/Interfaces/IAnalyticsProvider.h
+++ b/Engine/Source/Runtime/Analytics/Analytics/Public/Interfaces/IAnalyticsProvider.h
@@ -102,12 +102,27 @@ public:
 
 	/**
 	 * Flush any cached events to the analytics provider.
+	 * 
+	 * This does not guarantee that the flushed events have been fully committed to what ever storage
+	 * mechanism that the provider uses. The event might still be writing to disk or in-flight on the
+	 * network etc. This call will however guarantee that all currently cached events have started the
+	 *  process.
 	 *
 	 * Note that not all providers support explicitly sending any cached events. In which case this method
 	 * does nothing.
 	 */
 	virtual void FlushEvents() = 0;
 
+	/**
+	 * Functions the same as ::Flush except that it will attempt to wait until the events have been fully
+	 * committed to the provider.
+	 * 
+	 * @param InTimeoutSec	The maximum time allowed (in seconds) to wait for the events to become fully
+	 *						committed. If the events are not fully committed once this timer is exceeded the
+	 *						method will return anyway.
+	 */
+	virtual void BlockUntilFlushed(float InTimeoutSec) = 0;
+
 	/**
 	 * Set the UserID for use with analytics. Some providers require a unique ID
 	 * to be provided when supplying events, and some providers create their own.
diff --git a/Engine/Source/Runtime/Analytics/AnalyticsET/Public/IAnalyticsProviderET.h b/Engine/Source/Runtime/Analytics/AnalyticsET/Public/IAnalyticsProviderET.h
--- a/Engine/Source/Runtime/Analytics/AnalyticsET/Public/IAnalyticsProviderET.h
+++ b/Engine/Source/Runtime/Analytics/AnalyticsET/Public/IAnalyticsProviderET.h
@@ -110,11 +110,6 @@ public:
 	*/
 	virtual void SetEventCallback(const OnEventRecorded& Callback) = 0;
 
-	/**
-	* Blocks execution in the thread until all events have been flushed to the network.
-	*/
-	virtual void BlockUntilFlushed(float InTimeoutSec) = 0;
-
 	/**
 	 * Return the current provider configuration.
 	 */
diff --git a/Engine/Source/Runtime/Analytics/AnalyticsSwrve/Private/AnalyticsSwrve.cpp b/Engine/Source/Runtime/Analytics/AnalyticsSwrve/Private/AnalyticsSwrve.cpp
--- a/Engine/Source/Runtime/Analytics/AnalyticsSwrve/Private/AnalyticsSwrve.cpp
+++ b/Engine/Source/Runtime/Analytics/AnalyticsSwrve/Private/AnalyticsSwrve.cpp
@@ -26,8 +26,10 @@ public:
 
 	virtual bool StartSession(const TArray<FAnalyticsEventAttribute>& Attributes) override;
 	virtual void EndSession() override;
+	
 	/** Swrve PC implementation doesn't cache events */
 	virtual void FlushEvents() override {}
+	virtual void BlockUntilFlushed(float InTimeoutSec) override {}
 
 	virtual void SetUserID(const FString& InUserID) override;
 	virtual FString GetUserID() const override;
