Fix for buffer underrun in the Span and Region name generation

diff --git a/Engine/Plugins/Experimental/StudioTelemetry/Source/StudioTelemetry/Private/StudioTelemetryEditor.cpp b/Engine/Plugins/Experimental/StudioTelemetry/Source/StudioTelemetry/Private/StudioTelemetryEditor.cpp
--- a/Engine/Plugins/Experimental/StudioTelemetry/Source/StudioTelemetry/Private/StudioTelemetryEditor.cpp
+++ b/Engine/Plugins/Experimental/StudioTelemetry/Source/StudioTelemetry/Private/StudioTelemetryEditor.cpp
@@ -663,9 +663,7 @@ void FStudioTelemetryEditor::Initialize()
 			else
 			{
 				// Append the PIE transition count to the PIE name
-				FNameBuilder NameBuilder(PIEStartupSpanName);
-				NameBuilder.Append(TEXT("%d"), PIETransitionCount);
-				PIEStartupSpan = FStudioTelemetry::Get().StartSpan(FName(NameBuilder), PIESpan);
+				PIEStartupSpan = FStudioTelemetry::Get().StartSpan(FName(*FString::Printf(TEXT("%s%d"), *PIEStartupSpanName.ToString(), PIETransitionCount)), PIESpan);
 			}
 
 			TArray<FAnalyticsEventAttribute> Attributes;
diff --git a/Engine/Source/Runtime/Analytics/Analytics/Private/AnalyticsTracer.cpp b/Engine/Source/Runtime/Analytics/Analytics/Private/AnalyticsTracer.cpp
--- a/Engine/Source/Runtime/Analytics/Analytics/Private/AnalyticsTracer.cpp
+++ b/Engine/Source/Runtime/Analytics/Analytics/Private/AnalyticsTracer.cpp
@@ -195,7 +195,7 @@ TSharedPtr<IAnalyticsSpan> FAnalyticsTracer::StartSpan(const FName NewSpanName,
 
 void FAnalyticsTracer::BeginRegion(TSharedPtr<IAnalyticsSpan> Span)
 {
-	// This function is a temporary work around as UnrealInsights does not handle overalapping region with the same name.
+	// This function is a temporary work around as UnrealInsights does not handle overlapping region with the same name.
 	FName RegionName = Span->GetName();
 
 	uint32 NameCounter = 0;
@@ -204,8 +204,7 @@ void FAnalyticsTracer::BeginRegion(TSharedPtr<IAnalyticsSpan> Span)
 	{
 		// Generate a unique region name for this span
 		FNameBuilder NameBuilder(Span->GetName());
-		NameBuilder.Append(TEXT("%d"), NameCounter++);
-		RegionName = FName(NameBuilder);
+		RegionName = FName(FName(*FString::Printf(TEXT("%s%d"), *Span->GetName().ToString(), ++NameCounter)));
 	}
 
 	// Add the region name to the list by span ID
