[MP Locker] Re-add locker search analytics
* Move locker analytics into their own file
* Moved old loadout analytics code to new location so it's not in a huge shared header

#tests Tested that a breakpoint in the code was being hit in the correct circumstances and with the expected info.

diff --git a/Engine/Source/Runtime/Analytics/Analytics/Public/AnalyticsConversion.h b/Engine/Source/Runtime/Analytics/Analytics/Public/AnalyticsConversion.h
--- a/Engine/Source/Runtime/Analytics/Analytics/Public/AnalyticsConversion.h
+++ b/Engine/Source/Runtime/Analytics/Analytics/Public/AnalyticsConversion.h
@@ -79,6 +79,22 @@ FString AnalyticsConversionToString(const TArray<T, AllocatorType>& ValueArray)
 	return Result;
 }
 
+/** Array conversion. Creates comma-separated list. */
+template <typename T>
+FString AnalyticsConversionToString(const TArrayView<T>& ValueArray)
+{
+	FString Result;
+	// Serialize the array into "value1,value2,..." format
+	for (const T& Value : ValueArray)
+	{
+		Result += AnalyticsConversionToString(Value);
+		Result += TEXT(",");
+	}
+	// Remove the trailing comma (LeftChop will ensure an empty container won't crash here).
+	Result.LeftChopInline(1, EAllowShrinking::No);
+	return Result;
+}
+
 /** Map conversion. Creates comma-separated list. Creates comma-separated list with colon-separated key:value pairs. */
 template<typename KeyType, typename ValueType, typename Allocator, typename KeyFuncs>
 FString AnalyticsConversionToString(const TMap<KeyType, ValueType, Allocator, KeyFuncs>& ValueMap)
