Use a TransactionallySafeScopeLock for analytics

diff --git a/Engine/Source/Runtime/Analytics/AnalyticsET/Private/AnalyticsProviderETEventCache.cpp b/Engine/Source/Runtime/Analytics/AnalyticsET/Private/AnalyticsProviderETEventCache.cpp
--- a/Engine/Source/Runtime/Analytics/AnalyticsET/Private/AnalyticsProviderETEventCache.cpp
+++ b/Engine/Source/Runtime/Analytics/AnalyticsET/Private/AnalyticsProviderETEventCache.cpp
@@ -219,7 +219,7 @@ FAnalyticsProviderETEventCache::FAnalyticsProviderETEventCache(int32 InMaximumPa
 // We End with {"Events":[{"EventName":"<NAME>","DateOffset":"<OFFSET>",<DefaultAttrs>,<Attrs>}]}
 void FAnalyticsProviderETEventCache::AddToCache(FString EventName, const TArray<FAnalyticsEventAttribute>& Attributes)
 {
-	FScopeLock ScopedLock(&CachedEventsCS);
+	FTransactionallySafeScopeLock ScopedLock(&CachedEventsCS);
 
 	// If we estimate that 110% of the size estimate (in case there are a lot of Json escaping or multi-byte UTF8 chars) will exceed our max payload, queue up a flush. 
 	const int32 EventSizeEstimate = EventCacheStatic::ComputeEventSize(EventName, Attributes, CachedDefaultAttributeUTF8Stream.Num());
@@ -277,7 +277,7 @@ void FAnalyticsProviderETEventCache::AddToCache(FString EventName)
 
 void FAnalyticsProviderETEventCache::SetDefaultAttributes(TArray<FAnalyticsEventAttribute>&& DefaultAttributes)
 {
-	FScopeLock ScopedLock(&CachedEventsCS);
+	FTransactionallySafeScopeLock ScopedLock(&CachedEventsCS);
 
 	// store the array so we can return if if the user asks again.
 	CachedDefaultAttributes = MoveTemp(DefaultAttributes);
@@ -297,25 +297,25 @@ void FAnalyticsProviderETEventCache::SetDefaultAttributes(TArray<FAnalyticsEvent
 
 TArray<FAnalyticsEventAttribute> FAnalyticsProviderETEventCache::GetDefaultAttributes() const
 {
-	FScopeLock ScopedLock(&CachedEventsCS);
+	FTransactionallySafeScopeLock ScopedLock(&CachedEventsCS);
 	return CachedDefaultAttributes;
 }
 
 int32 FAnalyticsProviderETEventCache::GetDefaultAttributeCount() const
 {
-	FScopeLock ScopedLock(&CachedEventsCS);
+	FTransactionallySafeScopeLock ScopedLock(&CachedEventsCS);
 	return CachedDefaultAttributes.Num();
 }
 
 FAnalyticsEventAttribute FAnalyticsProviderETEventCache::GetDefaultAttribute(int32 AttributeIndex) const
 {
-	FScopeLock ScopedLock(&CachedEventsCS);
+	FTransactionallySafeScopeLock ScopedLock(&CachedEventsCS);
 	return CachedDefaultAttributes[AttributeIndex];
 }
 
 FString FAnalyticsProviderETEventCache::FlushCache(SIZE_T* OutEventCount)
 {
-	FScopeLock ScopedLock(&CachedEventsCS);
+	FTransactionallySafeScopeLock ScopedLock(&CachedEventsCS);
 	if (OutEventCount)
 	{
 		*OutEventCount = CachedEventEntries.Num();
@@ -327,7 +327,7 @@ FString FAnalyticsProviderETEventCache::FlushCache(SIZE_T* OutEventCount)
 
 TArray<uint8> FAnalyticsProviderETEventCache::FlushCacheUTF8()
 {
-	FScopeLock ScopedLock(&CachedEventsCS);
+	FTransactionallySafeScopeLock ScopedLock(&CachedEventsCS);
 
 	// if there's nothing queued up, flush what we have.
 	if (FlushQueue.Num() == 0 && CachedEventEntries.Num() > 0)
@@ -353,7 +353,7 @@ TArray<uint8> FAnalyticsProviderETEventCache::FlushCacheUTF8()
 void FAnalyticsProviderETEventCache::QueueFlush()
 {
 	const double StartTime = FPlatformTime::Seconds();
-	FScopeLock ScopedLock(&CachedEventsCS);
+	FTransactionallySafeScopeLock ScopedLock(&CachedEventsCS);
 
 	// early exit if nothing to flush.
 	if (CachedEventEntries.Num() == 0)
@@ -428,7 +428,7 @@ void FAnalyticsProviderETEventCache::QueueFlush()
 
 bool FAnalyticsProviderETEventCache::CanFlush() const
 {
-	FScopeLock ScopedLock(&CachedEventsCS);
+	FTransactionallySafeScopeLock ScopedLock(&CachedEventsCS);
 	return CachedEventEntries.Num() > 0 || FlushQueue.Num() > 0;
 }
 
@@ -439,7 +439,7 @@ bool FAnalyticsProviderETEventCache::HasFlushesQueued() const
 
 int FAnalyticsProviderETEventCache::GetNumCachedEvents() const
 {
-	FScopeLock ScopedLock(&CachedEventsCS);
+	FTransactionallySafeScopeLock ScopedLock(&CachedEventsCS);
 	return CachedEventEntries.Num();
 }
 
@@ -453,7 +453,7 @@ void FAnalyticsProviderETEventCache::SetPreallocatedPayloadSize(int32 InPrealloc
 	// if we are asking for a smaller buffer try to accommodate immediately.
 	if (PreallocatedPayloadSize < (int32)CachedEventUTF8Stream.GetAllocatedSize())
 	{
-		FScopeLock ScopedLock(&CachedEventsCS);
+		FTransactionallySafeScopeLock ScopedLock(&CachedEventsCS);
 		TArray<uint8> NewPayload;
 		NewPayload.Reserve(PreallocatedPayloadSize);
 		NewPayload = CachedEventUTF8Stream;
diff --git a/Engine/Source/Runtime/Analytics/AnalyticsET/Public/AnalyticsProviderETEventCache.h b/Engine/Source/Runtime/Analytics/AnalyticsET/Public/AnalyticsProviderETEventCache.h
--- a/Engine/Source/Runtime/Analytics/AnalyticsET/Public/AnalyticsProviderETEventCache.h
+++ b/Engine/Source/Runtime/Analytics/AnalyticsET/Public/AnalyticsProviderETEventCache.h
@@ -12,6 +12,7 @@
 #include "Misc/ScopeLock.h"
 #include "Templates/Function.h"
 #include "Templates/UnrealTemplate.h"
+#include "Misc/TransactionallySafeScopeLock.h"
 
 /**
  * Mixin class for Epic Telemetry implementors.
@@ -107,7 +108,7 @@ public:
 			:ScopedLock(&EventCache.CachedEventsCS)
 		{}
 	private:
-		FScopeLock ScopedLock;
+		FTransactionallySafeScopeLock ScopedLock;
 	};
 
 	static ANALYTICSET_API void OnStartupModule();
@@ -153,5 +154,5 @@ private:
 	TArray<TArray<uint8>> FlushQueue;
 
 	/** Critical section for updating the CachedEvents. Mutable to allow const methods to access the list. */
-	mutable FCriticalSection CachedEventsCS;
+	mutable FTransactionallySafeCriticalSection CachedEventsCS;
 };
