[AutoRTFM] Defer analytics reports that occur during a transaction to OnCommit time.

If a transaction abort occurs, this allows us to roll back the transaction safely without side-effects.

diff --git a/Engine/Source/Runtime/Analytics/AnalyticsET/Private/IAnalyticsProviderET.cpp b/Engine/Source/Runtime/Analytics/AnalyticsET/Private/IAnalyticsProviderET.cpp
--- a/Engine/Source/Runtime/Analytics/AnalyticsET/Private/IAnalyticsProviderET.cpp
+++ b/Engine/Source/Runtime/Analytics/AnalyticsET/Private/IAnalyticsProviderET.cpp
@@ -20,6 +20,7 @@
 #include "Misc/EngineVersion.h"
 #include "HttpRetrySystem.h"
 #include "HAL/IConsoleManager.h"
+#include "AutoRTFM/AutoRTFM.h"
 
 #include "AnalyticsPerfTracker.h"
 
@@ -431,15 +432,18 @@ TSharedRef<IHttpRequest, ESPMode::ThreadSafe> FAnalyticsProviderET::CreateReques
 
 void FAnalyticsProviderET::FlushEvents()
 {
-	QUICK_SCOPE_CYCLE_COUNTER(STAT_FAnalyticsProviderET_FlushEvents);
-	// Warn if this takes more than 2 ms
-	SCOPE_TIME_GUARD_MS(TEXT("FAnalyticsProviderET::FlushEvents"), 2);
-
-	// keep flushing until the event cache has cleared its queue.
-	while (EventCache.CanFlush())
+	UE_AUTORTFM_ONCOMMIT(this)
 	{
-		FlushEventsOnce();
-	}
+		QUICK_SCOPE_CYCLE_COUNTER(STAT_FAnalyticsProviderET_FlushEvents);
+		// Warn if this takes more than 2 ms
+		SCOPE_TIME_GUARD_MS(TEXT("FAnalyticsProviderET::FlushEvents"), 2);
+
+		// keep flushing until the event cache has cleared its queue.
+		while (EventCache.CanFlush())
+		{
+			FlushEventsOnce();
+		}
+	};
 }
 
 void FAnalyticsProviderET::FlushEventsOnce()
@@ -735,51 +739,56 @@ static inline void AnalyticsProviderETFlushEventLegacyHelper(FString& EventParam
 
 void FAnalyticsProviderET::FlushEventLegacy(const FString& EventName, const TArray<FAnalyticsEventAttribute>& Attributes)
 {
-	// this is a legacy pathway that doesn't accept batch payloads of cached data. We'll just send one request for each event, which will be slow for a large batch of requests at once.
-	if (ensure(FModuleManager::Get().IsModuleLoaded("HTTP")))
+	// If we are running transactionally, postpone the HTTP request until our transaction is committed.
+	// We need to explicitly copy the input arguments since they are passed in by reference.
+	UE_AUTORTFM_ONCOMMIT(this, EventName = EventName, Attributes = Attributes)
 	{
-		// first generate a payload from the eventand attributes
-		FString EventParams;
-		int PayloadNdx = 0;
-		for (int DefaultNdx = 0, NumDefaults = EventCache.GetDefaultAttributeCount(); DefaultNdx < NumDefaults; ++DefaultNdx)
-		{
-			AnalyticsProviderETFlushEventLegacyHelper(EventParams, PayloadNdx, EventCache.GetDefaultAttribute(DefaultNdx));
-			++PayloadNdx;
-		}
-		for (int AttrNdx = 0; AttrNdx < Attributes.Num(); ++AttrNdx)
+		// this is a legacy pathway that doesn't accept batch payloads of cached data. We'll just send one request for each event, which will be slow for a large batch of requests at once.
+		if (ensure(FModuleManager::Get().IsModuleLoaded("HTTP")))
 		{
-			AnalyticsProviderETFlushEventLegacyHelper(EventParams, PayloadNdx, Attributes[AttrNdx]);
-			++PayloadNdx;
-		}
+			// first generate a payload from the eventand attributes
+			FString EventParams;
+			int PayloadNdx = 0;
+			for (int DefaultNdx = 0, NumDefaults = EventCache.GetDefaultAttributeCount(); DefaultNdx < NumDefaults; ++DefaultNdx)
+			{
+				AnalyticsProviderETFlushEventLegacyHelper(EventParams, PayloadNdx, EventCache.GetDefaultAttribute(DefaultNdx));
+				++PayloadNdx;
+			}
+			for (int AttrNdx = 0; AttrNdx < Attributes.Num(); ++AttrNdx)
+			{
+				AnalyticsProviderETFlushEventLegacyHelper(EventParams, PayloadNdx, Attributes[AttrNdx]);
+				++PayloadNdx;
+			}
 
-		// log out the un-encoded values to make reading the log easier.
-		UE_LOG(LogAnalytics, VeryVerbose, TEXT("[%s] AnalyticsET URL:SendEvent.1?SessionID=%s&AppID=%s&AppVersion=%s&UserID=%s&EventName=%s%s"),
-			*Config.APIKeyET,
-			*SessionID,
-			*Config.APIKeyET,
-			*Config.AppVersionET,
-			*UserID,
-			*EventName,
-			*EventParams);
-
-		// Create/send Http request for an event
-		TSharedRef<IHttpRequest, ESPMode::ThreadSafe> HttpRequest = CreateRequest();
-		HttpRequest->SetHeader(TEXT("Content-Type"), TEXT("text/plain"));
-
-		// Don't need to URL encode the APIServer or the EventParams, which are already encoded, and contain parameter separaters that we DON'T want encoded.
-		FString URLPath = Config.APIServerET+Config.APIEndpointET;
-		URLPath += TEXT("SendEvent.1?SessionID=") + FPlatformHttp::UrlEncode(SessionID);
-		URLPath += TEXT("&AppID=") + FPlatformHttp::UrlEncode(Config.APIKeyET);
-		URLPath += TEXT("&AppVersion=") + FPlatformHttp::UrlEncode(Config.AppVersionET);
-		URLPath += TEXT("&UserID=") + FPlatformHttp::UrlEncode(UserID);
-		URLPath += TEXT("&EventName=") + FPlatformHttp::UrlEncode(EventName);
-		URLPath += EventParams;
-		HttpRequest->SetURL(URLPath);
-		HttpRequest->SetVerb(TEXT("GET"));
-		if (!bInDestructor)
-		{
-			HttpRequest->OnProcessRequestComplete().BindSP(this, &FAnalyticsProviderET::EventRequestComplete);
+			// log out the un-encoded values to make reading the log easier.
+			UE_LOG(LogAnalytics, VeryVerbose, TEXT("[%s] AnalyticsET URL:SendEvent.1?SessionID=%s&AppID=%s&AppVersion=%s&UserID=%s&EventName=%s%s"),
+				*Config.APIKeyET,
+				*SessionID,
+				*Config.APIKeyET,
+				*Config.AppVersionET,
+				*UserID,
+				*EventName,
+				*EventParams);
+
+			// Create/send Http request for an event
+			TSharedRef<IHttpRequest, ESPMode::ThreadSafe> HttpRequest = CreateRequest();
+			HttpRequest->SetHeader(TEXT("Content-Type"), TEXT("text/plain"));
+
+			// Don't need to URL encode the APIServer or the EventParams, which are already encoded, and contain parameter separaters that we DON'T want encoded.
+			FString URLPath = Config.APIServerET+Config.APIEndpointET;
+			URLPath += TEXT("SendEvent.1?SessionID=") + FPlatformHttp::UrlEncode(SessionID);
+			URLPath += TEXT("&AppID=") + FPlatformHttp::UrlEncode(Config.APIKeyET);
+			URLPath += TEXT("&AppVersion=") + FPlatformHttp::UrlEncode(Config.AppVersionET);
+			URLPath += TEXT("&UserID=") + FPlatformHttp::UrlEncode(UserID);
+			URLPath += TEXT("&EventName=") + FPlatformHttp::UrlEncode(EventName);
+			URLPath += EventParams;
+			HttpRequest->SetURL(URLPath);
+			HttpRequest->SetVerb(TEXT("GET"));
+			if (!bInDestructor)
+			{
+				HttpRequest->OnProcessRequestComplete().BindSP(this, &FAnalyticsProviderET::EventRequestComplete);
+			}
+			HttpRequest->ProcessRequest();
 		}
-		HttpRequest->ProcessRequest();
-	}
+	};
 }
