Add support for Analytics HTTP Requests to be sent to a Unix socket on platforms that support it

diff --git a/Engine/Source/Runtime/Analytics/AnalyticsET/Private/AnalyticsET.cpp b/Engine/Source/Runtime/Analytics/AnalyticsET/Private/AnalyticsET.cpp
--- a/Engine/Source/Runtime/Analytics/AnalyticsET/Private/AnalyticsET.cpp
+++ b/Engine/Source/Runtime/Analytics/AnalyticsET/Private/AnalyticsET.cpp
@@ -40,6 +40,7 @@ TSharedPtr<IAnalyticsProviderET> FAnalyticsET::CreateAnalyticsProviderET(const F
 		ConfigValues.APIEndpointET = GetConfigValue.Execute(Config::GetKeyNameForAPIEndpoint(), true);
 		ConfigValues.AppVersionET = GetConfigValue.Execute(Config::GetKeyNameForAppVersion(), false);
 		ConfigValues.UseLegacyProtocol = FCString::ToBool(*GetConfigValue.Execute(Config::GetKeyNameForUseLegacyProtocol(), false));
+		ConfigValues.APIUnixSocketPathET = GetConfigValue.Execute(Config::GetKeyNameForAPIUnixSocketPath(), false);
 
 		if (ConfigValues.AppVersionET.IsEmpty())
 		{
diff --git a/Engine/Source/Runtime/Analytics/AnalyticsET/Private/IAnalyticsProviderET.cpp b/Engine/Source/Runtime/Analytics/AnalyticsET/Private/IAnalyticsProviderET.cpp
--- a/Engine/Source/Runtime/Analytics/AnalyticsET/Private/IAnalyticsProviderET.cpp
+++ b/Engine/Source/Runtime/Analytics/AnalyticsET/Private/IAnalyticsProviderET.cpp
@@ -297,6 +297,13 @@ FAnalyticsProviderET::FAnalyticsProviderET(const FAnalyticsET::Config& ConfigVal
 		UE_LOG(LogAnalytics, Warning, TEXT("AnalyticsET: APIServerET is empty for APIKey (%s), creating as a NULL provider!"), *Config.APIKeyET);
 	}
 
+#if !UE_HTTP_SUPPORT_UNIX_SOCKET
+	if (!Config.APIUnixSocketPathET.IsEmpty())
+	{
+		UE_LOG(LogAnalytics, Warning, TEXT("[%s] Specified UnixSocketPath '%s' but that is not supported on this platform"), *Config.APIKeyET, *Config.APIUnixSocketPathET);
+	}
+#endif //UE_HTTP_SUPPORT_UNIX_SOCKET
+
 	// only need these if we are using the data router protocol.
 	if (!Config.UseLegacyProtocol)
 	{
@@ -504,6 +511,13 @@ void FAnalyticsProviderET::FlushEventsOnce()
 			HttpRequest->SetVerb(TEXT("POST"));
 			HttpRequest->SetContent(MoveTemp(Payload));
 
+#if UE_HTTP_SUPPORT_UNIX_SOCKET
+			if (!Config.APIUnixSocketPathET.IsEmpty())
+			{
+				HttpRequest->SetOption(HttpRequestOptions::UnixSocketPath, Config.APIUnixSocketPathET);
+			}
+#endif //UE_HTTP_SUPPORT_UNIX_SOCKET
+
 			// Don't set a response callback if we are in our destructor, as the instance will no longer be there to call.
 			if (!bInDestructor)
 			{
@@ -784,6 +798,14 @@ void FAnalyticsProviderET::FlushEventLegacy(const FString& EventName, const TArr
 			URLPath += EventParams;
 			HttpRequest->SetURL(URLPath);
 			HttpRequest->SetVerb(TEXT("GET"));
+
+#if UE_HTTP_SUPPORT_UNIX_SOCKET
+			if (!Config.APIUnixSocketPathET.IsEmpty())
+			{
+				HttpRequest->SetOption(HttpRequestOptions::UnixSocketPath, Config.APIUnixSocketPathET);
+			}
+#endif //UE_HTTP_SUPPORT_UNIX_SOCKET
+
 			if (!bInDestructor)
 			{
 				HttpRequest->OnProcessRequestComplete().BindSP(this, &FAnalyticsProviderET::EventRequestComplete);
diff --git a/Engine/Source/Runtime/Analytics/AnalyticsET/Public/AnalyticsET.h b/Engine/Source/Runtime/Analytics/AnalyticsET/Public/AnalyticsET.h
--- a/Engine/Source/Runtime/Analytics/AnalyticsET/Public/AnalyticsET.h
+++ b/Engine/Source/Runtime/Analytics/AnalyticsET/Public/AnalyticsET.h
@@ -53,6 +53,8 @@ public:
 		FString APIServerET;
 		/** ET API Endpoint - This is the API endpoint for the provider. */
 		FString APIEndpointET;
+		/** ET API UnixSocketPath - This is the UnixSocketPath to use as the target for the API endpoint for the provider. */
+		FString APIUnixSocketPathET;
 		/** ET Alt API Servers - Base URLs to send events on retry. */
 		TArray<FString> AltAPIServersET;
 		/** 
@@ -100,6 +102,8 @@ public:
 		static FString GetKeyNameForAPIServer() { return TEXT("APIServerET"); }
 		/** KeyName required for APIEndpoint configuration. */
 		static FString GetKeyNameForAPIEndpoint() { return TEXT("APIEndpointET"); }
+		/** KeyName required for APIEndpoint configuration. */
+		static FString GetKeyNameForAPIUnixSocketPath() { return TEXT("APIUnixSocketPathET"); }
 		/** KeyName required for AppVersion configuration. */
 		static FString GetKeyNameForAppVersion() { return TEXT("AppVersionET"); }
 		/** Optional parameter to use the legacy backend protocol. */
