[AutoRTFM] Add defensive guard against reaching FlushEventsOnce while AutoRTFM is in the closed.

This should be impossible, since this call can only be reached from `FlushEvents` (which wraps the call in OnCommit) and `Tick` (since a Ticker should not start in the closed).

diff --git a/Engine/Source/Runtime/Analytics/AnalyticsET/Private/IAnalyticsProviderET.cpp b/Engine/Source/Runtime/Analytics/AnalyticsET/Private/IAnalyticsProviderET.cpp
--- a/Engine/Source/Runtime/Analytics/AnalyticsET/Private/IAnalyticsProviderET.cpp
+++ b/Engine/Source/Runtime/Analytics/AnalyticsET/Private/IAnalyticsProviderET.cpp
@@ -455,6 +455,12 @@ void FAnalyticsProviderET::FlushEvents()
 
 void FAnalyticsProviderET::FlushEventsOnce()
 {
+	// FlushEventsOnce cannot be rolled back, so it must not occur inside an AutoRTFM transaction.
+	if (!ensure(!AutoRTFM::IsClosed()))
+	{
+		return;
+	}
+
 	// Make sure we don't try to flush too many times. When we are not caching events it's possible this can be called when there are no events in the array.
 	if (!EventCache.CanFlush())
 	{
