[LLM] Move more EngineMisc allocations to specific categories, mark more assets/objects for LLM tracking

#tests RR
#rnx

diff --git a/Engine/Source/Runtime/Analytics/AnalyticsET/Private/IAnalyticsProviderET.cpp b/Engine/Source/Runtime/Analytics/AnalyticsET/Private/IAnalyticsProviderET.cpp
--- a/Engine/Source/Runtime/Analytics/AnalyticsET/Private/IAnalyticsProviderET.cpp
+++ b/Engine/Source/Runtime/Analytics/AnalyticsET/Private/IAnalyticsProviderET.cpp
@@ -216,6 +216,7 @@ TSharedPtr<IAnalyticsProviderET> FAnalyticsET::CreateAnalyticsProvider(const Con
 		UE_LOG(LogAnalytics, Warning, TEXT("CreateAnalyticsProvider config not contain required parameter %s"), *Config::GetKeyNameForAPIKey());
 		return NULL;
 	}
+	LLM_SCOPE_BYNAME(TEXT("Analytics"));
 	return MakeShared<FAnalyticsProviderET>(ConfigValues);
 }
 
@@ -327,6 +328,7 @@ FAnalyticsProviderET::FAnalyticsProviderET(const FAnalyticsET::Config& ConfigVal
 
 bool FAnalyticsProviderET::Tick(float DeltaSeconds)
 {
+	LLM_SCOPE_BYNAME(TEXT("Analytics"));
 	QUICK_SCOPE_CYCLE_COUNTER(STAT_FAnalyticsProviderET_Tick);
 
 PRAGMA_DISABLE_DEPRECATION_WARNINGS
@@ -382,6 +384,7 @@ FAnalyticsProviderET::~FAnalyticsProviderET()
 
 bool FAnalyticsProviderET::StartSession(FString InSessionID, const TArray<FAnalyticsEventAttribute>& Attributes)
 {
+	LLM_SCOPE_BYNAME(TEXT("Analytics"));
 	UE_LOG(LogAnalytics, Log, TEXT("[%s] AnalyticsET::StartSession"), *Config.APIKeyET);
 
 	// end/flush previous session before staring new one
@@ -404,6 +407,7 @@ bool FAnalyticsProviderET::StartSession(FString InSessionID, const TArray<FAnaly
  */
 void FAnalyticsProviderET::EndSession()
 {
+	LLM_SCOPE_BYNAME(TEXT("Analytics"));
 	if (bSessionInProgress)
 	{
 		RecordEvent(TEXT("SessionEnd"), TArray<FAnalyticsEventAttribute>());
@@ -418,6 +422,7 @@ void FAnalyticsProviderET::EndSession()
 
 TSharedRef<IHttpRequest, ESPMode::ThreadSafe> FAnalyticsProviderET::CreateRequest()
 {
+	LLM_SCOPE_BYNAME(TEXT("Analytics"));
 	if (!ensure(FModuleManager::Get().IsModuleLoaded("HTTP")))
 	{
 		UE_LOG(LogAnalytics, Display, TEXT("[%s] ET Analytics provider tried to create a new HTTP request when HTTP was shutdown"), *Config.APIKeyET);
@@ -461,6 +466,7 @@ void FAnalyticsProviderET::FlushEventsOnce()
 		return;
 	}
 
+	LLM_SCOPE_BYNAME(TEXT("Analytics"));
 	// Make sure we don't try to flush too many times. When we are not caching events it's possible this can be called when there are no events in the array.
 	if (!EventCache.CanFlush())
 	{
@@ -607,6 +613,7 @@ bool FAnalyticsProviderET::ShouldRecordEvent(const FString& EventName) const
 
 void FAnalyticsProviderET::RecordEvent(FString&& EventName, const TArray<FAnalyticsEventAttribute>& Attributes)
 {
+	LLM_SCOPE_BYNAME(TEXT("Analytics"));
 	// let higher level code filter the decision of whether to send the event
 	if (ShouldRecordEvent(EventName))
 	{
@@ -689,6 +696,7 @@ void FAnalyticsProviderET::SetURLEndpoint(const FString& UrlEndpoint, const TArr
 	{
 		return;
 	}
+	LLM_SCOPE_BYNAME(TEXT("Analytics"));
 
 	// flush existing events before changing URL domains.
 	FlushEvents();
@@ -727,6 +735,7 @@ void FAnalyticsProviderET::SetURLEndpoint(const FString& UrlEndpoint, const TArr
 
 void FAnalyticsProviderET::SetHeader(const FString& HeaderName, const FString& HeaderValue)
 {
+	LLM_SCOPE_BYNAME(TEXT("Analytics"));
 	if (HeaderValue.IsEmpty())
 	{
 		HttpHeaders.Remove(HeaderName);
@@ -739,6 +748,7 @@ void FAnalyticsProviderET::SetHeader(const FString& HeaderName, const FString& H
 
 void FAnalyticsProviderET::BlockUntilFlushed(float InTimeoutSec)
 {
+	LLM_SCOPE_BYNAME(TEXT("Analytics"));
 	FlushEvents();
 	HttpRetryManager->BlockUntilFlushed(InTimeoutSec);
 }
@@ -766,6 +776,7 @@ void FAnalyticsProviderET::FlushEventLegacy(const FString& EventName, const TArr
 		// this is a legacy pathway that doesn't accept batch payloads of cached data. We'll just send one request for each event, which will be slow for a large batch of requests at once.
 		if (ensure(FModuleManager::Get().IsModuleLoaded("HTTP")))
 		{
+			LLM_SCOPE_BYNAME(TEXT("Analytics"));
 			// first generate a payload from the eventand attributes
 			FString EventParams;
 			int PayloadNdx = 0;
diff --git a/Engine/Source/Runtime/CoreUObject/Private/UObject/ScriptCore.cpp b/Engine/Source/Runtime/CoreUObject/Private/UObject/ScriptCore.cpp
--- a/Engine/Source/Runtime/CoreUObject/Private/UObject/ScriptCore.cpp
+++ b/Engine/Source/Runtime/CoreUObject/Private/UObject/ScriptCore.cpp
@@ -20,6 +20,7 @@
 #include "UObject/Object.h"
 #include "UObject/CoreNative.h"
 #include "UObject/Class.h"
+#include "UObject/Package.h"
 #include "Templates/Casts.h"
 #include "Serialization/NullArchive.h"
 #include "UObject/SoftObjectPtr.h"
@@ -32,6 +33,9 @@
 #include "UObject/ScriptMacros.h"
 #include "UObject/UObjectThreadContext.h"
 #include "HAL/IConsoleManager.h"
+#include "HAL/LowLevelMemTracker.h"
+#include "HAL/LowLevelMemStats.h"
+#include "ProfilingDebugging/AssetMetadataTrace.h"
 #include "AutoRTFM/AutoRTFM.h"
 
 DEFINE_LOG_CATEGORY(LogScriptFrame);
@@ -1085,6 +1089,10 @@ void UObject::CallFunction( FFrame& Stack, RESULT_DECL, UFunction* Function )
 #endif // PER_FUNCTION_SCRIPT_STATS
 
 	SCOPE_CYCLE_UOBJECT(ContextScope, GVerboseScriptStats ? this : nullptr);
+	LLM_SCOPE(ELLMTag::UObject);
+	LLM_SCOPE_DYNAMIC_STAT_OBJECTPATH(GetPackage(), ELLMTagSet::Assets);
+	LLM_SCOPE_DYNAMIC_STAT_OBJECTPATH(GetClass(), ELLMTagSet::AssetClasses);
+	UE_TRACE_METADATA_SCOPE_ASSET(this, GetClass());
 
 	checkSlow(Function);
 
@@ -1262,18 +1270,36 @@ void ProcessLocalScriptFunction(UObject* Context, FFrame& Stack, RESULT_DECL)
 void ProcessLocalFunction(UObject* Context, UFunction* Fn, FFrame& Stack, RESULT_DECL)
 {
 	checkSlow(Fn);
-	if(Fn->HasAnyFunctionFlags(FUNC_Native))
+
+	auto ContinueProcessLocalFuntionInner = [&]() {
+		if(Fn->HasAnyFunctionFlags(FUNC_Native))
+		{
+			FScopeCycleCounterUObject NativeContextScope(GVerboseScriptStats ? Context : nullptr);
+			Fn->Invoke(Context, Stack, RESULT_PARAM);
+		}
+		else
+		{
+	#if PER_FUNCTION_SCRIPT_STATS
+			const bool bShouldTrackFunction = (Stack.DepthCounter <= GMaxFunctionStatDepth);
+			SCOPE_CYCLE_UOBJECT(FunctionScope, bShouldTrackFunction ? Fn : nullptr);
+	#endif // PER_FUNCTION_SCRIPT_STATS
+			ProcessScriptFunction(Context, Fn, Stack, RESULT_PARAM, ProcessLocalScriptFunction);
+		}
+	};
+
+#if ENABLE_LOW_LEVEL_MEM_TRACKER
+	if (Context && FLowLevelMemTracker::IsEnabled())
 	{
-		FScopeCycleCounterUObject NativeContextScope(GVerboseScriptStats ? Context : nullptr);
-		Fn->Invoke(Context, Stack, RESULT_PARAM);
+		LLM_SCOPE(ELLMTag::UObject);
+		LLM_SCOPE_DYNAMIC_STAT_OBJECTPATH(Context->GetPackage(), ELLMTagSet::Assets);
+		LLM_SCOPE_DYNAMIC_STAT_OBJECTPATH(Context->GetClass(), ELLMTagSet::AssetClasses);
+		UE_TRACE_METADATA_SCOPE_ASSET(Context, Context->GetClass());
+		return ContinueProcessLocalFuntionInner();
 	}
 	else
+#endif
 	{
-#if PER_FUNCTION_SCRIPT_STATS
-		const bool bShouldTrackFunction = (Stack.DepthCounter <= GMaxFunctionStatDepth);
-		SCOPE_CYCLE_UOBJECT(FunctionScope, bShouldTrackFunction ? Fn : nullptr);
-#endif // PER_FUNCTION_SCRIPT_STATS
-		ProcessScriptFunction(Context, Fn, Stack, RESULT_PARAM, ProcessLocalScriptFunction);
+		return ContinueProcessLocalFuntionInner();
 	}
 }
 
@@ -2024,6 +2050,11 @@ void UObject::ProcessEvent( UFunction* Function, void* Parms )
 
 	SCOPE_CYCLE_UOBJECT(ContextScope, GVerboseScriptStats ? this : nullptr);
 
+	LLM_SCOPE(ELLMTag::UObject);
+	LLM_SCOPE_DYNAMIC_STAT_OBJECTPATH(GetPackage(), ELLMTagSet::Assets);
+	LLM_SCOPE_DYNAMIC_STAT_OBJECTPATH(GetClass(), ELLMTagSet::AssetClasses);
+	UE_TRACE_METADATA_SCOPE_ASSET(this, GetClass());
+
 #if LIGHTWEIGHT_PROCESS_EVENT_COUNTER
 	TGuardValue<int32> PECounter(ProcessEventCounter, ProcessEventCounter + 1);
 	CONDITIONAL_SCOPE_CYCLE_COUNTER(STAT_BlueprintTime, IsInGameThread() && ProcessEventCounter == 1);
diff --git a/Engine/Source/Runtime/Experimental/Chaos/Private/Chaos/Framework/PhysicsSolverBase.cpp b/Engine/Source/Runtime/Experimental/Chaos/Private/Chaos/Framework/PhysicsSolverBase.cpp
--- a/Engine/Source/Runtime/Experimental/Chaos/Private/Chaos/Framework/PhysicsSolverBase.cpp
+++ b/Engine/Source/Runtime/Experimental/Chaos/Private/Chaos/Framework/PhysicsSolverBase.cpp
@@ -14,6 +14,7 @@
 
 #include "ChaosVisualDebugger/ChaosVisualDebuggerTrace.h"
 #include "ChaosVisualDebugger/ChaosVDContextProvider.h"
+#include "HAL/LowLevelMemTracker.h"
 
 DEFINE_STAT(STAT_AsyncPullResults);
 DEFINE_STAT(STAT_AsyncInterpolateResults);
@@ -420,6 +421,7 @@ namespace Chaos
 
 	FGraphEventRef FPhysicsSolverBase::AdvanceAndDispatch_External(FReal InDt)
 	{
+		LLM_SCOPE(ELLMTag::ChaosScene);
 		const bool bSubstepping = MMaxSubSteps > 1;
 		SetSolverSubstep_External(bSubstepping);
 		const FReal DtWithPause = bPaused_External ? 0.0f : InDt;
diff --git a/Engine/Source/Runtime/Experimental/IoStore/OnDemand/Private/OnDemandIoStore.cpp b/Engine/Source/Runtime/Experimental/IoStore/OnDemand/Private/OnDemandIoStore.cpp
--- a/Engine/Source/Runtime/Experimental/IoStore/OnDemand/Private/OnDemandIoStore.cpp
+++ b/Engine/Source/Runtime/Experimental/IoStore/OnDemand/Private/OnDemandIoStore.cpp
@@ -21,6 +21,7 @@
 #include "HAL/PlatformFileManager.h"
 #include "HAL/PlatformTime.h"
 #include "HAL/PlatformProcess.h"
+#include "HAL/LowLevelMemTracker.h"
 #include "IO/IoChunkEncoding.h"
 #include "IO/IoContainerHeader.h"
 #include "IO/IoStoreOnDemand.h"
@@ -34,6 +35,7 @@
 #include "Misc/Guid.h"
 #include "Misc/Paths.h"
 #include "Misc/PathViews.h"
+#include "ProfilingDebugging/AssetMetadataTrace.h"
 #include "Serialization/MemoryReader.h"
 #if !(UE_BUILD_SHIPPING|UE_BUILD_TEST)
 #include "String/LexFromString.h"
@@ -491,6 +493,7 @@ FOnDemandIoStore::~FOnDemandIoStore()
 
 FIoStatus FOnDemandIoStore::Initialize()
 {
+	LLM_SCOPE(ELLMTag::FileSystem);
 	const FIoStatus CacheStatus = InitializeOnDemandInstallCache();
 	if (CacheStatus.GetErrorCode() == EIoErrorCode::PendingFork)
 	{
@@ -607,6 +610,7 @@ void FOnDemandIoStore::Install(
 	FOnDemandInstallProgressed&& OnProgress /*= nullptr*/,
 	const FOnDemandCancellationToken* CancellationToken)
 {
+	LLM_SCOPE(ELLMTag::FileSystem);
 	FSharedInstallRequest InstallRequest = MakeShared<FInstallRequest>();
 	InstallRequest->Args				= MoveTemp(Args);
 	InstallRequest->OnCompleted			= MoveTemp(OnCompleted);
@@ -623,6 +627,7 @@ void FOnDemandIoStore::Install(
 
 void FOnDemandIoStore::Purge(FOnDemandPurgeArgs&& Args, FOnDemandPurgeCompleted&& OnCompleted)
 {
+	LLM_SCOPE(ELLMTag::FileSystem);
 	FSharedPurgeRequest PurgeRequest = MakeShared<FPurgeRequest>();
 	PurgeRequest->Args			= MoveTemp(Args);
 	PurgeRequest->OnCompleted	= MoveTemp(OnCompleted);
@@ -637,6 +642,7 @@ void FOnDemandIoStore::Purge(FOnDemandPurgeArgs&& Args, FOnDemandPurgeCompleted&
 
 FIoStatus FOnDemandIoStore::Unmount(FStringView MountId)
 {
+	LLM_SCOPE(ELLMTag::FileSystem);
 	UE_LOG(LogIoStoreOnDemand, Log, TEXT("Unmounting '%s'"), *WriteToString<256>(MountId));
 
 	bool bPendingMount = false;
@@ -941,6 +947,7 @@ void FOnDemandIoStore::TryEnterTickLoop()
 
 void FOnDemandIoStore::TickLoop()
 {
+	LLM_SCOPE(ELLMTag::FileSystem);
 	ON_SCOPE_EXIT { UE_LOG(LogIoStoreOnDemand, Verbose, TEXT("Exiting I/O store tick loop")); };
 
 	UE_LOG(LogIoStoreOnDemand, Verbose, TEXT("Entering I/O store tick loop"));
@@ -1688,6 +1695,11 @@ void FOnDemandIoStore::CreateContainersFromToc(
 	FOnDemandToc& Toc,
 	TArray<FSharedOnDemandContainer>& Out)
 {
+	static FName AssetClassName(TEXT("OnDemandIoStore"));
+	LLM_SCOPE(ELLMTag::FileSystem);
+	LLM_TAGSET_SCOPE(FName(MountId), ELLMTagSet::Assets);
+	LLM_TAGSET_SCOPE(AssetClassName, ELLMTagSet::AssetClasses);
+	UE_TRACE_METADATA_SCOPE_ASSET_FNAME(FName(MountId), AssetClassName, FName(TocPath));
 	const FOnDemandTocHeader& Header = Toc.Header;
 	const FName CompressionFormat(Header.CompressionFormat);
 
@@ -1780,6 +1792,10 @@ void FOnDemandIoStore::CreateContainersFromToc(
 			}
 		}
 
+		// Do not count output container memory as being allocated for utocs
+		LLM_TAGSET_SCOPE_CLEAR(ELLMTagSet::Assets);
+		LLM_TAGSET_SCOPE_CLEAR(ELLMTagSet::AssetClasses);
+		UE_TRACE_METADATA_CLEAR_SCOPE();
 		Out.Add(MoveTemp(Container));
 	}
 }
@@ -1903,6 +1919,7 @@ FIoStatus FOnDemandIoStore::GetContainersAndPackagesForInstall(
 
 void FOnDemandIoStore::ReleaseContent(FOnDemandInternalContentHandle& ContentHandle)
 {
+	LLM_SCOPE(ELLMTag::FileSystem);
 	UE_LOG(LogIoStoreOnDemand, Log, TEXT("Releasing content handle '%s'"), *LexToString(ContentHandle));
 
 	UE::TUniqueLock Lock(ContainerMutex);
