Fixed EditorTelemetry and RuntimeTelemetry module startup so that it does not cause an expception when booting in an  Editorshipping build.
Removed the superfluous Destroy messages from AnalyticsProviderET

diff --git a/Engine/Plugins/Experimental/EditorTelemetry/Source/Private/EditorTelemetry.cpp b/Engine/Plugins/Experimental/EditorTelemetry/Source/Private/EditorTelemetry.cpp
--- a/Engine/Plugins/Experimental/EditorTelemetry/Source/Private/EditorTelemetry.cpp
+++ b/Engine/Plugins/Experimental/EditorTelemetry/Source/Private/EditorTelemetry.cpp
@@ -430,7 +430,7 @@ void FEditorTelemetry::HeartbeatCallback()
 
 void FEditorTelemetry::StartSession()
 {
-	if (FStudioTelemetry::Get().IsAvailable() == false)
+	if (FStudioTelemetry::Get().IsSessionRunning() == false)
 	{
 		return;
 	}
@@ -1019,5 +1019,9 @@ void FEditorTelemetry::StartSession()
 
 void FEditorTelemetry::EndSession()
 {
-	FStudioTelemetry::Get().EndSpan(EditorSpan);
+	if (EditorSpan.IsValid())
+	{
+		FStudioTelemetry::Get().EndSpan(EditorSpan);
+		EditorSpan.Reset();
+	}
 }
diff --git a/Engine/Plugins/Experimental/EditorTelemetry/Source/Private/EditorTelemetryModule.cpp b/Engine/Plugins/Experimental/EditorTelemetry/Source/Private/EditorTelemetryModule.cpp
--- a/Engine/Plugins/Experimental/EditorTelemetry/Source/Private/EditorTelemetryModule.cpp
+++ b/Engine/Plugins/Experimental/EditorTelemetry/Source/Private/EditorTelemetryModule.cpp
@@ -9,14 +9,22 @@ IMPLEMENT_MODULE(FEditorTelemetryModule, EditorTelemetry);
 
 void FEditorTelemetryModule::StartupModule()
 {
+	if (FStudioTelemetry::Get().IsSessionRunning())
+	{
+		// StudioTelemetry session has already started
+		FEditorTelemetry::Get().StartSession();
+	}
+		
 	FStudioTelemetry::Get().GetOnStartSession().AddLambda([this]()
 		{
+			// StudioTelemetry session has just started
 			FEditorTelemetry::Get().StartSession();
 		}
 	);
 	
 	FStudioTelemetry::Get().GetOnEndSession().AddLambda([this]()
 		{
+			// StudioTelemetry session is about to end
 			FEditorTelemetry::Get().EndSession();
 		}
 	);
diff --git a/Engine/Plugins/Experimental/RuntimeTelemetry/Source/Private/RuntimeTelemetryModule.cpp b/Engine/Plugins/Experimental/RuntimeTelemetry/Source/Private/RuntimeTelemetryModule.cpp
--- a/Engine/Plugins/Experimental/RuntimeTelemetry/Source/Private/RuntimeTelemetryModule.cpp
+++ b/Engine/Plugins/Experimental/RuntimeTelemetry/Source/Private/RuntimeTelemetryModule.cpp
@@ -9,19 +9,22 @@ IMPLEMENT_MODULE(FRuntimeTelemetryModule, RuntimeTelemetry);
 
 void FRuntimeTelemetryModule::StartupModule()
 {
-	if (FStudioTelemetry::Get().IsAvailable())
+	if (FStudioTelemetry::Get().IsSessionRunning())
 	{
+		// StudioTelemetry session has already started
 		FRuntimeTelemetry::Get().StartSession();
 	}
 
 	FStudioTelemetry::Get().GetOnStartSession().AddLambda([this]()
 		{
+			// StudioTelemetry session has just started
 			FRuntimeTelemetry::Get().StartSession();
 		}
 	);
 	
 	FStudioTelemetry::Get().GetOnEndSession().AddLambda([this]()
 		{
+			// StudioTelemetry session is about to end
 			FRuntimeTelemetry::Get().EndSession();
 		}
 	);
diff --git a/Engine/Source/Runtime/Analytics/AnalyticsET/Private/IAnalyticsProviderET.cpp b/Engine/Source/Runtime/Analytics/AnalyticsET/Private/IAnalyticsProviderET.cpp
--- a/Engine/Source/Runtime/Analytics/AnalyticsET/Private/IAnalyticsProviderET.cpp
+++ b/Engine/Source/Runtime/Analytics/AnalyticsET/Private/IAnalyticsProviderET.cpp
@@ -372,10 +372,8 @@ bool FAnalyticsProviderET::Tick(float DeltaSeconds)
 
 FAnalyticsProviderET::~FAnalyticsProviderET()
 {
-	UE_LOG(LogAnalytics, Display, TEXT("[%s] Destroying ET Analytics provider"), *Config.APIKeyET);
 	bInDestructor = true;
 	EndSession();
-	UE_LOG(LogAnalytics, Display, TEXT("[%s] Destroyed ET Analytics provider"), *Config.APIKeyET);
 }
 
 bool FAnalyticsProviderET::StartSession(FString InSessionID, const TArray<FAnalyticsEventAttribute>& Attributes)
@@ -407,7 +405,7 @@ void FAnalyticsProviderET::EndSession()
 	if (bSessionInProgress)
 	{
 		RecordEvent(TEXT("SessionEnd"), TArray<FAnalyticsEventAttribute>());
-		UE_LOG(LogAnalytics, Display, TEXT("[%s] Ended ET Analytics provider session"), *Config.APIKeyET);
+		UE_LOG(LogAnalytics, Display, TEXT("[%s] AnalyticsET::EndSession"), *Config.APIKeyET);
 	}
 
 	FlushEvents();
