Removed dependency of Horde Analytics Provider on the Editor Only Horde library
This change allows us to send Horde Analytics from Runtime sessions to Horde as necessary

diff --git a/Engine/Source/Runtime/Analytics/AnalyticsHorde/AnalyticsHorde.Build.cs b/Engine/Source/Runtime/Analytics/AnalyticsHorde/AnalyticsHorde.Build.cs
--- a/Engine/Source/Runtime/Analytics/AnalyticsHorde/AnalyticsHorde.Build.cs
+++ b/Engine/Source/Runtime/Analytics/AnalyticsHorde/AnalyticsHorde.Build.cs
@@ -13,16 +13,7 @@ namespace UnrealBuildTool.Rules
 				"Analytics",
 				"AnalyticsET"
 				}
-			);
-
-			if (Target.bBuildEditor)
-			{
-					PrivateDependencyModuleNames.AddRange(
-						new string[] {
-						"Horde"
-						}
-					);
-				}
-			}
+			);	
+		}
 	}
 }
diff --git a/Engine/Source/Runtime/Analytics/AnalyticsHorde/Private/AnalyticsHorde.cpp b/Engine/Source/Runtime/Analytics/AnalyticsHorde/Private/AnalyticsHorde.cpp
--- a/Engine/Source/Runtime/Analytics/AnalyticsHorde/Private/AnalyticsHorde.cpp
+++ b/Engine/Source/Runtime/Analytics/AnalyticsHorde/Private/AnalyticsHorde.cpp
@@ -3,10 +3,8 @@
 #include "AnalyticsHorde.h"
 #include "AnalyticsET.h"
 #include "IAnalyticsProviderET.h"
-
-#if WITH_EDITOR
-#include "Horde.h"
-#endif
+#include "Misc/CommandLine.h"
+#include "Misc/Parse.h"
 
 IMPLEMENT_MODULE( FAnalyticsHorde, AnalyticsHorde );
 
@@ -18,22 +16,37 @@ void FAnalyticsHorde::ShutdownModule()
 {
 }
 
+FString GetServerURL()
+{
+	static FString ServerURL;
+
+	if (ServerURL.IsEmpty())
+	{
+		if (false == FParse::Value(FCommandLine::Get(), TEXT("HordeServerUrl="), ServerURL))
+		{
+			ServerURL = FPlatformMisc::GetEnvironmentVariable(TEXT("UE_HORDE_URL"));
+		}
+	}
+	return ServerURL;
+}
+
 TSharedPtr<IAnalyticsProvider> FAnalyticsHorde::CreateAnalyticsProvider(const FAnalyticsProviderConfigurationDelegate& GetConfigValue) const
 {
 	if (GetConfigValue.IsBound())
 	{
 		TSharedPtr<IAnalyticsProviderET> AnalyicsProviderET = FAnalyticsET::Get().CreateAnalyticsProviderET(GetConfigValue);
 
-#if WITH_EDITOR
-		// Check if we have specifcied a Horde URL in the environment. This allows jobs running on Horde to send telemetry directly to the server they were run on
-		const FString HordeServerURL = FHorde::GetServerURL();
-
-		if (AnalyicsProviderET.IsValid() && !HordeServerURL.IsEmpty())
+		if (AnalyicsProviderET.IsValid())
 		{
-			TArray<FString> AltDomains;
-			AnalyicsProviderET->SetURLEndpoint(HordeServerURL, AltDomains);
+			// Check if we have specified a Horde URL in the environment. This allows jobs running on Horde to send telemetry directly to the server they were run on
+			const FString HordeServerURL = GetServerURL();
+
+			if (!HordeServerURL.IsEmpty())
+			{
+				TArray<FString> AltDomains;
+				AnalyicsProviderET->SetURLEndpoint(HordeServerURL, AltDomains);
+			}
 		}
-#endif
 		
 		return AnalyicsProviderET;
 	}
