Log Analytics Provider picks up the UE_TELEMETRY path if an explicit one has not been provided.

diff --git a/Engine/Source/Runtime/Analytics/AnalyticsLog/Private/AnalyticsProviderLog.cpp b/Engine/Source/Runtime/Analytics/AnalyticsLog/Private/AnalyticsProviderLog.cpp
--- a/Engine/Source/Runtime/Analytics/AnalyticsLog/Private/AnalyticsProviderLog.cpp
+++ b/Engine/Source/Runtime/Analytics/AnalyticsLog/Private/AnalyticsProviderLog.cpp
@@ -21,8 +21,14 @@ FAnalyticsProviderLog::FAnalyticsProviderLog(const FAnalyticsProviderConfigurati
 
 	if (FolderPath.IsEmpty())
 	{
-		// Use default output path
-		FolderPath = FPaths::ProjectSavedDir() / TEXT("Telemetry");
+		// See if there's a folder specified in the environment 
+		FolderPath = FPlatformMisc::GetEnvironmentVariable(TEXT("UE_TELEMETRY_DIR"));
+
+		if (FolderPath.IsEmpty())
+		{
+			// Use default output path
+			FolderPath = FPaths::ProjectSavedDir() / TEXT("Telemetry");
+		}
 	}
 
 	// Create the full output path
