Show the URL the Analytics Provider has connected to in the log.

diff --git a/Engine/Source/Runtime/Analytics/AnalyticsET/Private/IAnalyticsProviderET.cpp b/Engine/Source/Runtime/Analytics/AnalyticsET/Private/IAnalyticsProviderET.cpp
--- a/Engine/Source/Runtime/Analytics/AnalyticsET/Private/IAnalyticsProviderET.cpp
+++ b/Engine/Source/Runtime/Analytics/AnalyticsET/Private/IAnalyticsProviderET.cpp
@@ -725,6 +725,10 @@ void FAnalyticsProviderET::SetURLEndpoint(const FString& UrlEndpoint, const TArr
 	{
 		UE_LOG(LogAnalytics, Warning, TEXT("AnalyticsET: APIServerET is empty for APIKey (%s), converting to a NULL provider!"), *Config.APIKeyET);
 	}
+	else
+	{
+		UE_LOG(LogAnalytics, Log, TEXT("AnalyticsET: Set APIServerET to %s"), *Config.APIServerET);
+	}
 }
 
 void FAnalyticsProviderET::SetHeader(const FString& HeaderName, const FString& HeaderValue)
