- Add UrlEndpoint override for ProviderType/Environment pairs
- Add callback that fires immediately before the HTTP request for an event is processed.

diff --git a/Engine/Source/Runtime/Analytics/AnalyticsET/Private/IAnalyticsProviderET.cpp b/Engine/Source/Runtime/Analytics/AnalyticsET/Private/IAnalyticsProviderET.cpp
--- a/Engine/Source/Runtime/Analytics/AnalyticsET/Private/IAnalyticsProviderET.cpp
+++ b/Engine/Source/Runtime/Analytics/AnalyticsET/Private/IAnalyticsProviderET.cpp
@@ -149,10 +149,13 @@ public:
 	virtual FAnalyticsEventAttribute GetDefaultEventAttribute(int AttributeIndex) const override;
 	virtual void SetEventCallback(const OnEventRecorded& Callback) override;
 
-	virtual void SetURLEndpoint(const FString& UrlEndpoint, const TArray<FString>& AltDomains) override;
+	virtual void SetUrlDomain(const FString& Domain, const TArray<FString>& AltDomains) override;
+	virtual void SetUrlPath(const FString& Path) override;
 	virtual void SetHeader(const FString& HeaderName, const FString& HeaderValue) override;
 	virtual void BlockUntilFlushed(float InTimeoutSec) override;
 	virtual void SetShouldRecordEventFunc(const ShouldRecordEventFunction& InShouldRecordEventFunc) override;
+	virtual FOnPreAnalyticsEventProcessed& OnPreAnalyticsEventProcessed() override { return OnPreAnalyticsEventProcessedDelegate; }
+
 	virtual ~FAnalyticsProviderET();
 
 	virtual const FAnalyticsET::Config& GetConfig() const override { return Config; }
@@ -206,6 +209,8 @@ private:
 	TMap<FString, FString> HttpHeaders;
 
 	FAnalyticsUserAgentCache UserAgentCache;
+
+	FOnPreAnalyticsEventProcessed OnPreAnalyticsEventProcessedDelegate;
 };
 
 TSharedPtr<IAnalyticsProviderET> FAnalyticsET::CreateAnalyticsProvider(const Config& ConfigValues) const
@@ -536,6 +541,8 @@ void FAnalyticsProviderET::FlushEventsOnce()
 				HttpRequest->OnProcessRequestComplete().BindSP(this, &FAnalyticsProviderET::EventRequestComplete);
 			}
 
+			OnPreAnalyticsEventProcessedDelegate.Broadcast(HttpRequest);
+
 			HttpRequest->ProcessRequest();
 		}
 	}
@@ -689,10 +696,10 @@ void FAnalyticsProviderET::EventRequestComplete(FHttpRequestPtr HttpRequest, FHt
 	}
 }
 
-void FAnalyticsProviderET::SetURLEndpoint(const FString& UrlEndpoint, const TArray<FString>& AltDomains)
+void FAnalyticsProviderET::SetUrlDomain(const FString& Domain, const TArray<FString>& AltDomains)
 {
 	// See if anything is actually changing before going through the work to flush and reset the URLs.
-	if (Config.APIServerET == UrlEndpoint && Config.AltAPIServersET == AltDomains)
+	if (Config.APIServerET == Domain && Config.AltAPIServersET == AltDomains)
 	{
 		return;
 	}
@@ -701,7 +708,7 @@ void FAnalyticsProviderET::SetURLEndpoint(const FString& UrlEndpoint, const TArr
 	// flush existing events before changing URL domains.
 	FlushEvents();
 
-	Config.APIServerET = UrlEndpoint;
+	Config.APIServerET = Domain;
 	Config.AltAPIServersET = AltDomains;
 
 	// Set the number of retries to the number of retry URLs that have been passed in.
@@ -737,6 +744,21 @@ void FAnalyticsProviderET::SetURLEndpoint(const FString& UrlEndpoint, const TArr
 	}
 }
 
+void FAnalyticsProviderET::SetUrlPath(const FString& Path)
+{
+	// See if anything is actually changing before going through the work to flush and reset the URLs.
+	if (Config.APIEndpointET == Path)
+	{
+		return;
+	}
+	LLM_SCOPE_BYNAME(TEXT("Analytics"));
+
+	// flush existing events before changing URL path.
+	FlushEvents();
+
+	Config.APIEndpointET = Path;
+}
+
 void FAnalyticsProviderET::SetHeader(const FString& HeaderName, const FString& HeaderValue)
 {
 	LLM_SCOPE_BYNAME(TEXT("Analytics"));
@@ -810,7 +832,7 @@ void FAnalyticsProviderET::FlushEventLegacy(const FString& EventName, const TArr
 			HttpRequest->SetHeader(TEXT("Content-Type"), TEXT("text/plain"));
 
 			// Don't need to URL encode the APIServer or the EventParams, which are already encoded, and contain parameter separaters that we DON'T want encoded.
-			FString URLPath = Config.APIServerET+Config.APIEndpointET;
+			FString URLPath = Config.APIServerET + Config.APIEndpointET;
 			URLPath += TEXT("SendEvent.1?SessionID=") + FPlatformHttp::UrlEncode(SessionID);
 			URLPath += TEXT("&AppID=") + FPlatformHttp::UrlEncode(Config.APIKeyET);
 			URLPath += TEXT("&AppVersion=") + FPlatformHttp::UrlEncode(Config.AppVersionET);
@@ -831,6 +853,9 @@ void FAnalyticsProviderET::FlushEventLegacy(const FString& EventName, const TArr
 			{
 				HttpRequest->OnProcessRequestComplete().BindSP(this, &FAnalyticsProviderET::EventRequestComplete);
 			}
+
+			OnPreAnalyticsEventProcessedDelegate.Broadcast(HttpRequest);
+
 			HttpRequest->ProcessRequest();
 		}
 	};
diff --git a/Engine/Source/Runtime/Analytics/AnalyticsET/Public/IAnalyticsProviderET.h b/Engine/Source/Runtime/Analytics/AnalyticsET/Public/IAnalyticsProviderET.h
--- a/Engine/Source/Runtime/Analytics/AnalyticsET/Public/IAnalyticsProviderET.h
+++ b/Engine/Source/Runtime/Analytics/AnalyticsET/Public/IAnalyticsProviderET.h
@@ -14,6 +14,11 @@
 #include "Templates/Function.h"
 #include "Templates/UnrealTemplate.h"
 
+class IHttpRequest;
+
+/** Fired immediately prior to the HTTP request for an analytics event being processed by HTTP */
+DECLARE_MULTICAST_DELEGATE_TwoParams(FOnPreAnalyticsEventProcessed, TSharedRef<IHttpRequest, ESPMode::ThreadSafe>& /*HttpRequest*/);
+
 /** ET specific analytics provider instance. Exposes additional APIs to support Json-based events, default attributes, and allowing events to be disabled (generally via hotfixing). */
 class IAnalyticsProviderET : public IAnalyticsProvider
 {
@@ -93,9 +98,15 @@ public:
 	virtual void RecordEvent(FString&& EventName, const TArray<FAnalyticsEventAttribute>& Attributes) = 0;
 
 	/**
-	 * Updates the default URL endpoint and AltDomains.
+	 * Updates the default Domain and AltDomains.
 	 */
-	virtual void SetURLEndpoint(const FString& UrlEndpoint, const TArray<FString>& AltDomains) = 0;
+	virtual void SetUrlDomain(const FString& Domain, const TArray<FString>& AltDomains) = 0;
+
+	/**
+	* Updates the Path
+	*/
+	virtual void SetUrlPath(const FString& Path) = 0;
+
 	/**
 	 * Sets a header to be included with analytics http requests
 	 */
@@ -125,6 +136,8 @@ public:
 
 	/** Set an event filter to dynamically control whether an event should be sent. */
 	virtual void SetShouldRecordEventFunc(const ShouldRecordEventFunction& ShouldRecordEventFunc) = 0;
+	
+	virtual FOnPreAnalyticsEventProcessed& OnPreAnalyticsEventProcessed() = 0;
 
 private:
 	/** Needed to support the old, unsafe GetDefaultAttributes() API. */
diff --git a/Engine/Source/Runtime/Analytics/AnalyticsHorde/Private/AnalyticsHorde.cpp b/Engine/Source/Runtime/Analytics/AnalyticsHorde/Private/AnalyticsHorde.cpp
--- a/Engine/Source/Runtime/Analytics/AnalyticsHorde/Private/AnalyticsHorde.cpp
+++ b/Engine/Source/Runtime/Analytics/AnalyticsHorde/Private/AnalyticsHorde.cpp
@@ -44,7 +44,7 @@ TSharedPtr<IAnalyticsProvider> FAnalyticsHorde::CreateAnalyticsProvider(const FA
 			if (!HordeServerURL.IsEmpty())
 			{
 				TArray<FString> AltDomains;
-				AnalyicsProviderET->SetURLEndpoint(HordeServerURL, AltDomains);
+				AnalyicsProviderET->SetUrlDomain(HordeServerURL, AltDomains);
 			}
 		}
 		
