Implemented check for duplicate attributes when recording analytics events

#rnx

diff --git a/Engine/Source/Runtime/Analytics/Analytics/Public/Interfaces/IAnalyticsProvider.h b/Engine/Source/Runtime/Analytics/Analytics/Public/Interfaces/IAnalyticsProvider.h
--- a/Engine/Source/Runtime/Analytics/Analytics/Public/Interfaces/IAnalyticsProvider.h
+++ b/Engine/Source/Runtime/Analytics/Analytics/Public/Interfaces/IAnalyticsProvider.h
@@ -445,5 +445,23 @@ public:
 	}
 
 	virtual ~IAnalyticsProvider() {}
+
+protected:
+	/**
+	 * Checks for duplicate key names in the attribute array. Only implemented for debug and development builds.
+	 */
+	void CheckForDuplicateAttributes(const FString& EventName, const TArray<FAnalyticsEventAttribute>& Attributes) const
+	{
+#if UE_BUILD_DEBUG || UE_BUILD_DEVELOPMENT
+		// Expose events that have duplicate attribute names. This is is not handled by the analytics backends in any reliable way.
+		for (int32 index0 = 0; index0 < Attributes.Num(); ++index0)
+		{
+			for (int32 index1 = index0 + 1; index1 < Attributes.Num(); ++index1)
+			{
+				checkf(Attributes[index0].GetName() != Attributes[index1].GetName(), TEXT("Duplicate Attributes Found For Event %s %s==%s"), *EventName, *Attributes[index0].GetName(), *Attributes[index1].GetName());
+			}
+		}
+#endif
+	}
 };
 
diff --git a/Engine/Source/Runtime/Analytics/AnalyticsET/Private/IAnalyticsProviderET.cpp b/Engine/Source/Runtime/Analytics/AnalyticsET/Private/IAnalyticsProviderET.cpp
--- a/Engine/Source/Runtime/Analytics/AnalyticsET/Private/IAnalyticsProviderET.cpp
+++ b/Engine/Source/Runtime/Analytics/AnalyticsET/Private/IAnalyticsProviderET.cpp
@@ -619,6 +619,9 @@ bool FAnalyticsProviderET::ShouldRecordEvent(const FString& EventName) const
 void FAnalyticsProviderET::RecordEvent(FString&& EventName, const TArray<FAnalyticsEventAttribute>& Attributes)
 {
 	LLM_SCOPE_BYNAME(TEXT("Analytics"));
+
+	CheckForDuplicateAttributes(EventName, Attributes);
+
 	// let higher level code filter the decision of whether to send the event
 	if (ShouldRecordEvent(EventName))
 	{
diff --git a/Engine/Source/Runtime/StudioTelemetry/Private/AnalyticsProviderBroadcast.cpp b/Engine/Source/Runtime/StudioTelemetry/Private/AnalyticsProviderBroadcast.cpp
--- a/Engine/Source/Runtime/StudioTelemetry/Private/AnalyticsProviderBroadcast.cpp
+++ b/Engine/Source/Runtime/StudioTelemetry/Private/AnalyticsProviderBroadcast.cpp
@@ -261,16 +261,7 @@ void FAnalyticsProviderBroadcast::SetRecordEventCallback(OnRecordEvent Callback)
 
 void FAnalyticsProviderBroadcast::RecordEvent(const FString& EventName, const TArray<FAnalyticsEventAttribute>& Attributes)
 {
-#if UE_BUILD_DEBUG || UE_BUILD_DEVELOPMENT
-	// Expose events that have duplicate attribute names. This is is not handled by the analytics backends in any reliable way.
-	for (int32 index0 = 0; index0 < Attributes.Num(); ++index0)
-	{
-		for (int32 index1 = index0 + 1; index1 < Attributes.Num(); ++index1)
-		{
-			checkf(Attributes[index0].GetName() != Attributes[index1].GetName(), TEXT("Duplicate Attributes Found For Event %s %s==%s"), *EventName, *Attributes[index0].GetName(), *Attributes[index1].GetName());
-		}
-	}
-#endif
+	CheckForDuplicateAttributes(EventName, Attributes);
 
 	for (TProviders::TConstIterator it(Providers); it; ++it)
 	{
