Duplicate attributes detection now logs a warning instead of using a check

#rnx

diff --git a/Engine/Source/Runtime/Analytics/Analytics/Public/Interfaces/IAnalyticsProvider.h b/Engine/Source/Runtime/Analytics/Analytics/Public/Interfaces/IAnalyticsProvider.h
--- a/Engine/Source/Runtime/Analytics/Analytics/Public/Interfaces/IAnalyticsProvider.h
+++ b/Engine/Source/Runtime/Analytics/Analytics/Public/Interfaces/IAnalyticsProvider.h
@@ -3,6 +3,7 @@
 #pragma once
 
 #include "AnalyticsEventAttribute.h"
+#include "Analytics.h"
 #include "CoreMinimal.h"
 
 class Error;
@@ -450,7 +451,7 @@ protected:
 	/**
 	 * Checks for duplicate key names in the attribute array. Only implemented for debug and development builds.
 	 */
-	void CheckForDuplicateAttributes(const FString& EventName, const TArray<FAnalyticsEventAttribute>& Attributes) const
+	static void CheckForDuplicateAttributes(const FString& EventName, const TArray<FAnalyticsEventAttribute>& Attributes)
 	{
 #if UE_BUILD_DEBUG || UE_BUILD_DEVELOPMENT
 		// Expose events that have duplicate attribute names. This is is not handled by the analytics backends in any reliable way.
@@ -458,7 +459,10 @@ protected:
 		{
 			for (int32 index1 = index0 + 1; index1 < Attributes.Num(); ++index1)
 			{
-				checkf(Attributes[index0].GetName() != Attributes[index1].GetName(), TEXT("Duplicate Attributes Found For Event %s %s==%s"), *EventName, *Attributes[index0].GetName(), *Attributes[index1].GetName());
+				if (Attributes[index0].GetName() == Attributes[index1].GetName())
+				{
+					UE_LOG(LogAnalytics, Warning, TEXT("Duplicate Attributes Found For Event %s %s==%s"), *EventName, *Attributes[index0].GetName(), *Attributes[index1].GetName());
+				}
 			}
 		}
 #endif
