Proper fix for analytics check() being replaced with a log. Moved definition of the logging function to its own source file and removed duplicated static definitions of the log category.

#rnx

diff --git a/Engine/Plugins/Runtime/Analytics/Adjust/Source/AndroidAdjust/Private/AndroidAdjust.cpp b/Engine/Plugins/Runtime/Analytics/Adjust/Source/AndroidAdjust/Private/AndroidAdjust.cpp
--- a/Engine/Plugins/Runtime/Analytics/Adjust/Source/AndroidAdjust/Private/AndroidAdjust.cpp
+++ b/Engine/Plugins/Runtime/Analytics/Adjust/Source/AndroidAdjust/Private/AndroidAdjust.cpp
@@ -1,6 +1,7 @@
 // Copyright Epic Games, Inc. All Rights Reserved.
 
 #include "AndroidAdjust.h"
+#include "Analytics.h"
 #include "Modules/ModuleManager.h"
 #include "AndroidAdjustProvider.h"
 #include "Android/AndroidJNI.h"
@@ -8,8 +9,6 @@
 #include "Misc/Paths.h"
 #include "Misc/ConfigCacheIni.h"
 
-DEFINE_LOG_CATEGORY_STATIC(LogAnalytics, Display, All);
-
 IMPLEMENT_MODULE( FAnalyticsAndroidAdjust, AndroidAdjust )
 
 TSharedPtr<IAnalyticsProvider> FAnalyticsProviderAdjust::Provider;
diff --git a/Engine/Plugins/Runtime/Analytics/Adjust/Source/IOSAdjust/Private/IOSAdjust.cpp b/Engine/Plugins/Runtime/Analytics/Adjust/Source/IOSAdjust/Private/IOSAdjust.cpp
--- a/Engine/Plugins/Runtime/Analytics/Adjust/Source/IOSAdjust/Private/IOSAdjust.cpp
+++ b/Engine/Plugins/Runtime/Analytics/Adjust/Source/IOSAdjust/Private/IOSAdjust.cpp
@@ -1,14 +1,13 @@
 // Copyright Epic Games, Inc. All Rights Reserved.
 
 #include "IOSAdjust.h"
+#include "Analytics.h"
 #include "IOSAdjustProvider.h"
 #include "Misc/Paths.h"
 #include "Misc/ConfigCacheIni.h"
 
 #import <AdjustSdk/Adjust.h>
 
-DEFINE_LOG_CATEGORY_STATIC(LogAnalytics, Display, All);
-
 IMPLEMENT_MODULE( FAnalyticsIOSAdjust, IOSAdjust )
 
 TSharedPtr<IAnalyticsProvider> FAnalyticsProviderAdjust::Provider;
diff --git a/Engine/Source/Runtime/Analytics/Analytics/Private/Interfaces/IAnalyticsProvider.cpp b/Engine/Source/Runtime/Analytics/Analytics/Private/Interfaces/IAnalyticsProvider.cpp
new file mode 100644
--- /dev/null
+++ b/Engine/Source/Runtime/Analytics/Analytics/Private/Interfaces/IAnalyticsProvider.cpp
@@ -0,0 +1,21 @@
+// Copyright Epic Games, Inc. All Rights Reserved.
+
+#include "Interfaces/IAnalyticsProvider.h"
+#include "Analytics.h"
+
+void IAnalyticsProvider::CheckForDuplicateAttributes(const FString& EventName, const TArray<FAnalyticsEventAttribute>& Attributes)
+{
+#if UE_BUILD_DEBUG || UE_BUILD_DEVELOPMENT
+	// Expose events that have duplicate attribute names. This is is not handled by the analytics backends in any reliable way.
+	for (int32 index0 = 0; index0 < Attributes.Num(); ++index0)
+	{
+		for (int32 index1 = index0 + 1; index1 < Attributes.Num(); ++index1)
+		{
+			if (Attributes[index0].GetName() == Attributes[index1].GetName())
+			{
+				UE_LOG(LogAnalytics, Warning, TEXT("Duplicate Attributes Found For Event %s %s==%s"), *EventName, *Attributes[index0].GetName(), *Attributes[index1].GetName());
+			}
+		}
+	}
+#endif
+}
diff --git a/Engine/Source/Runtime/Analytics/Analytics/Public/Interfaces/IAnalyticsProvider.h b/Engine/Source/Runtime/Analytics/Analytics/Public/Interfaces/IAnalyticsProvider.h
--- a/Engine/Source/Runtime/Analytics/Analytics/Public/Interfaces/IAnalyticsProvider.h
+++ b/Engine/Source/Runtime/Analytics/Analytics/Public/Interfaces/IAnalyticsProvider.h
@@ -450,18 +450,7 @@ protected:
 	/**
 	 * Checks for duplicate key names in the attribute array. Only implemented for debug and development builds.
 	 */
-	void CheckForDuplicateAttributes(const FString& EventName, const TArray<FAnalyticsEventAttribute>& Attributes) const
-	{
-#if UE_BUILD_DEBUG || UE_BUILD_DEVELOPMENT
-		// Expose events that have duplicate attribute names. This is is not handled by the analytics backends in any reliable way.
-		for (int32 index0 = 0; index0 < Attributes.Num(); ++index0)
-		{
-			for (int32 index1 = index0 + 1; index1 < Attributes.Num(); ++index1)
-			{
-				checkf(Attributes[index0].GetName() != Attributes[index1].GetName(), TEXT("Duplicate Attributes Found For Event %s %s==%s"), *EventName, *Attributes[index0].GetName(), *Attributes[index1].GetName());
-			}
-		}
-#endif
-	}
+	void ANALYTICS_API CheckForDuplicateAttributes(const FString& EventName, const TArray<FAnalyticsEventAttribute>& Attributes);
+
 };
 
