- Hoist the default ApiEndpoint (path) from the constructor to a GetDefault method, so it can be used during hotfixes

#rnx

diff --git a/Engine/Source/Runtime/Analytics/AnalyticsET/Private/IAnalyticsProviderET.cpp b/Engine/Source/Runtime/Analytics/AnalyticsET/Private/IAnalyticsProviderET.cpp
--- a/Engine/Source/Runtime/Analytics/AnalyticsET/Private/IAnalyticsProviderET.cpp
+++ b/Engine/Source/Runtime/Analytics/AnalyticsET/Private/IAnalyticsProviderET.cpp
@@ -296,10 +296,9 @@ FAnalyticsProviderET::FAnalyticsProviderET(const FAnalyticsET::Config& ConfigVal
 		? FString(FApp::GetBuildVersion())
 		: ConfigAppVersion.Replace(TEXT("%VERSION%"), FApp::GetBuildVersion(), ESearchCase::CaseSensitive);
 
-	if ( Config.APIEndpointET.IsEmpty() )
+	if (Config.APIEndpointET.IsEmpty())
 	{
-		// Set a default API Endpoint
-		Config.APIEndpointET = TEXT("datarouter/api/v1/public/data");
+		Config.APIEndpointET = FAnalyticsET::Config::GetDefaultAPIEndpoint();
 	}
 
 	if (Config.APIServerET.IsEmpty())
diff --git a/Engine/Source/Runtime/Analytics/AnalyticsET/Public/AnalyticsET.h b/Engine/Source/Runtime/Analytics/AnalyticsET/Public/AnalyticsET.h
--- a/Engine/Source/Runtime/Analytics/AnalyticsET/Public/AnalyticsET.h
+++ b/Engine/Source/Runtime/Analytics/AnalyticsET/Public/AnalyticsET.h
@@ -116,6 +116,8 @@ public:
 		static FString GetDefaultAppEnvironment() { return TEXT("datacollector-binary"); }
 		/** Default value if no UploadType is given, and UseDataRouter protocol is specified. */
 		static FString GetDefaultUploadType() { return TEXT("eteventstream"); }
+		/** Default value if no APIEndpoint is given. */
+		static FString GetDefaultAPIEndpoint() { return TEXT("datarouter/api/v1/public/data"); }
 	};
 
 	//--------------------------------------------------------------------------
