[Insights] Add some trace regions to long running engine events (PIE, LoadMap, Async Compilation in editor)...

diff --git a/Engine/Plugins/TraceUtilities/Source/EditorTraceUtilities/EditorTraceUtilities.Build.cs b/Engine/Plugins/TraceUtilities/Source/EditorTraceUtilities/EditorTraceUtilities.Build.cs
--- a/Engine/Plugins/TraceUtilities/Source/EditorTraceUtilities/EditorTraceUtilities.Build.cs
+++ b/Engine/Plugins/TraceUtilities/Source/EditorTraceUtilities/EditorTraceUtilities.Build.cs
@@ -25,6 +25,7 @@ public class EditorTraceUtilities : ModuleRules
 		PrivateDependencyModuleNames.AddRange(
 			new string[] {
 				"Engine",
+				"UnrealEd",
 				"InputCore",
 				"Slate",
 				"SlateCore",
diff --git a/Engine/Plugins/TraceUtilities/Source/EditorTraceUtilities/Private/EditorTraceUtilities.cpp b/Engine/Plugins/TraceUtilities/Source/EditorTraceUtilities/Private/EditorTraceUtilities.cpp
--- a/Engine/Plugins/TraceUtilities/Source/EditorTraceUtilities/Private/EditorTraceUtilities.cpp
+++ b/Engine/Plugins/TraceUtilities/Source/EditorTraceUtilities/Private/EditorTraceUtilities.cpp
@@ -2,6 +2,7 @@
 
 #include "EditorTraceUtilities.h"
 
+#include "Editor.h"
 #include "Misc/ConfigContext.h"
 #include "ToolMenus.h"
 #include "InsightsStatusBar/SInsightsStatusBar.h"
@@ -23,6 +24,10 @@ void FEditorTraceUtilitiesModule::StartupModule()
 
 	RegisterInsightsStatusWidgetWithToolMenu();
 	FConfigContext::ReadIntoGConfig().Load(TEXT("TraceUtilities"), EditorTraceUtilitiesIni);
+#if MISCTRACE_ENABLED
+	FEditorDelegates::BeginPIE.AddLambda([](bool){ TRACE_BEGIN_REGION(TEXT("PIE"), TEXT("EditorModes")); });
+	FEditorDelegates::EndPIE.AddLambda([](bool){ TRACE_END_REGION(TEXT("PIE")); });
+#endif
 }
 
 void FEditorTraceUtilitiesModule::ShutdownModule()
diff --git a/Engine/Source/Editor/UnrealEd/Private/Kismet2/DebuggerCommands.cpp b/Engine/Source/Editor/UnrealEd/Private/Kismet2/DebuggerCommands.cpp
--- a/Engine/Source/Editor/UnrealEd/Private/Kismet2/DebuggerCommands.cpp
+++ b/Engine/Source/Editor/UnrealEd/Private/Kismet2/DebuggerCommands.cpp
@@ -1528,8 +1528,6 @@ void FInternalPlayWorldCommandCallbacks::PlayInViewport_Clicked()
 		return;
 	}
 
-	TRACE_BOOKMARK(TEXT("PIE"));
-
 	FLevelEditorModule& LevelEditorModule = FModuleManager::GetModuleChecked<FLevelEditorModule>(TEXT("LevelEditor"));
 
 	/** Set PlayInViewPort as the last executed play command */
diff --git a/Engine/Source/Editor/UnrealEd/Private/NavigationBuildingNotification.cpp b/Engine/Source/Editor/UnrealEd/Private/NavigationBuildingNotification.cpp
--- a/Engine/Source/Editor/UnrealEd/Private/NavigationBuildingNotification.cpp
+++ b/Engine/Source/Editor/UnrealEd/Private/NavigationBuildingNotification.cpp
@@ -44,6 +44,7 @@ void FNavigationBuildingNotificationImpl::BuildStarted()
 	NavigationBuildNotificationPtr = FSlateNotificationManager::Get().AddNotification(Info);
 	if (NavigationBuildNotificationPtr.IsValid())
 	{
+		TRACE_BEGIN_REGION(TEXT("Building Navigation"), TEXT("EditorNotifications"));
 		NavigationBuildNotificationPtr.Pin()->SetCompletionState(SNotificationItem::CS_Pending);
 	}
 }
@@ -85,7 +86,7 @@ void FNavigationBuildingNotificationImpl::BuildFinished()
 				}
 			}
 		}
-
+		TRACE_END_REGION(TEXT("Building Navigation"));
 		NavigationBuildNotificationPtr.Reset();
 	}
 
diff --git a/Engine/Source/Runtime/Analytics/Analytics/Private/AnalyticsTracer.cpp b/Engine/Source/Runtime/Analytics/Analytics/Private/AnalyticsTracer.cpp
--- a/Engine/Source/Runtime/Analytics/Analytics/Private/AnalyticsTracer.cpp
+++ b/Engine/Source/Runtime/Analytics/Analytics/Private/AnalyticsTracer.cpp
@@ -1,7 +1,6 @@
 // Copyright Epic Games, Inc. All Rights Reserved.
 
 #include "AnalyticsTracer.h"
-#include "ProfilingDebugging/MiscTrace.h"
 
 UE_DISABLE_OPTIMIZATION_SHIP
  
@@ -200,38 +199,23 @@ TSharedPtr<IAnalyticsSpan> FAnalyticsTracer::StartSpan(const FName NewSpanName,
 
 void FAnalyticsTracer::BeginRegion(TSharedPtr<IAnalyticsSpan> Span)
 {
-	// This function is a temporary work around as UnrealInsights does not handle overlapping region with the same name.
-	FName RegionName = Span->GetName();
-
-	uint32 NameCounter = 0;
-
-	while (RegionNames.Find(RegionName) != nullptr)
-	{
-		// Generate a unique region name for this span
-		FNameBuilder NameBuilder(Span->GetName());
-		RegionName = FName(FName(*FString::Printf(TEXT("%s%d"), *Span->GetName().ToString(), ++NameCounter)));
-	}
-
-	// Add the region name to the list by span ID
-	RegionNames.Emplace(RegionName, Span->GetId());
-
-	TRACE_BEGIN_REGION(*RegionName.ToString());
+#if MISCTRACE_ENABLED
+	const FName& RegionName = Span->GetName();
+	uint64 RegionId = TRACE_BEGIN_REGION_WITH_ID(*RegionName.ToString(), TEXT("Analytics"));
+	// Add the region id to the list by span ID
+	RegionIds.Emplace(Span->GetId(), RegionId);
+#endif
 }
 
 void FAnalyticsTracer::EndRegion(TSharedPtr<IAnalyticsSpan> Span)
 {
-	// Find the region, end it and remove from the region names
-	for (TMap<FName, FGuid>::TConstIterator it(RegionNames); it; ++it)
+#if MISCTRACE_ENABLED
+	if (uint64* RegionId = RegionIds.Find(Span->GetId()))
 	{
-		// Slow match by ID on removal, fast match on creation
-		if ((*it).Value == Span->GetId())
-		{
-			const FName& RegionName = (*it).Key;
-			TRACE_END_REGION(*RegionName.ToString());
-			RegionNames.Remove(RegionName);
-			return;
-		}
+		TRACE_END_REGION_WITH_ID(*RegionId);
+		RegionIds.Remove(Span->GetId());
 	}
+#endif
 }
 
 bool FAnalyticsTracer::StartSpanInternal(TSharedPtr<IAnalyticsSpan> Span, const TArray<FAnalyticsEventAttribute>& AdditionalAttributes)
diff --git a/Engine/Source/Runtime/Analytics/Analytics/Public/AnalyticsTracer.h b/Engine/Source/Runtime/Analytics/Analytics/Public/AnalyticsTracer.h
--- a/Engine/Source/Runtime/Analytics/Analytics/Public/AnalyticsTracer.h
+++ b/Engine/Source/Runtime/Analytics/Analytics/Public/AnalyticsTracer.h
@@ -3,6 +3,7 @@
 #pragma once
 
 #include "Interfaces/IAnalyticsTracer.h"
+#include "ProfilingDebugging/MiscTrace.h"
 
 using FThreadId = uint32;
 
@@ -77,7 +78,9 @@ private:
 	TSharedPtr<IAnalyticsSpan>						SessionSpan;		// The root span, this will always be present in an active session
 	TArray<TSharedPtr<IAnalyticsSpan>>				ActiveSpanStack;	// Stack of active spans as WeakPtrs
 	TMap<FGuid, TArray<TWeakPtr<IAnalyticsSpan>>>	SpanHeirarchy;		// List of child spans per parent span
-	TMap<FName, FGuid>								RegionNames;		// Temprorary fix for duplicate span names in regions	
+#if MISCTRACE_ENABLED
+	TMap<FGuid, uint64>								RegionIds;		// Ids of active timing regions
+#endif
 	FCriticalSection								CriticalSection;	
 };
 
diff --git a/Engine/Source/Runtime/Core/Private/Misc/SlowTask.cpp b/Engine/Source/Runtime/Core/Private/Misc/SlowTask.cpp
--- a/Engine/Source/Runtime/Core/Private/Misc/SlowTask.cpp
+++ b/Engine/Source/Runtime/Core/Private/Misc/SlowTask.cpp
@@ -13,8 +13,8 @@
 #include "Misc/AssertionMacros.h"
 #include "Misc/FeedbackContext.h"
 #include "Misc/SlowTaskStack.h"
-#include "ProfilingDebugging/MiscTrace.h"
 
+#if MISCTRACE_ENABLED
 static int32 GSlowTaskMaxTraceRegionDepth = 2;
 static FAutoConsoleVariableRef CVarSlowTaskMaxTraceRegionDepth (
 	TEXT("Trace.SlowTaskMaxRegionDepth"),
@@ -22,6 +22,7 @@ static FAutoConsoleVariableRef CVarSlowTaskMaxTraceRegionDepth (
 	TEXT("Maximum depth of nested slow tasks to create as trace regions in Insights"),
 	ECVF_Default
 );
+#endif
 
 bool FSlowTask::ShouldCreateThrottledSlowTask()
 {
@@ -49,7 +50,9 @@ FSlowTask::FSlowTask(float InAmountOfWork, const FText& InDefaultMessage, bool b
 	, bCreatedDialog(false)		// only set to true if we create a dialog
 	, Context(InContext)
 	, bSkipRecursiveDialogCreation(false)
+#if MISCTRACE_ENABLED
 	, RegionId(0)
+#endif
 {
 	// If we have no work to do ourselves, create an arbitrary scope so that any actions performed underneath this still contribute to this one.
 	if (TotalAmountOfWork == 0.f)
@@ -104,7 +107,7 @@ void FSlowTask::Initialize()
 #if MISCTRACE_ENABLED
 		if (Context.ScopeStack.Num() <= GSlowTaskMaxTraceRegionDepth)
 		{
-			RegionId = TRACE_BEGIN_REGION_WITH_ID(!DefaultMessage.IsEmpty() ? *DefaultMessage.ToString() : TEXT("<SlowTask>"));
+			RegionId = TRACE_BEGIN_REGION_WITH_ID(!DefaultMessage.IsEmpty() ? *DefaultMessage.ToString() : TEXT("<SlowTask>"), TEXT("SlowTask"));
 		}
 #endif
 	}
diff --git a/Engine/Source/Runtime/Core/Public/Misc/SlowTask.h b/Engine/Source/Runtime/Core/Public/Misc/SlowTask.h
--- a/Engine/Source/Runtime/Core/Public/Misc/SlowTask.h
+++ b/Engine/Source/Runtime/Core/Public/Misc/SlowTask.h
@@ -6,6 +6,7 @@
 #include "CoreTypes.h"
 #include "Internationalization/Text.h"
 #include "Misc/Optional.h"
+#include "ProfilingDebugging/MiscTrace.h"
 
 class FFeedbackContext;
 
@@ -71,9 +72,10 @@ private:
 	/** Avoid calling hierarchical dialog creation if not required */
 	bool bSkipRecursiveDialogCreation : 1;
 
+#if MISCTRACE_ENABLED
 	/** Id of the traced timing region */
 	uint64 RegionId;
-
+#endif
 	/** Prevent copying */
 	FSlowTask(const FSlowTask&);
 
diff --git a/Engine/Source/Runtime/Engine/Private/LevelStreaming.cpp b/Engine/Source/Runtime/Engine/Private/LevelStreaming.cpp
--- a/Engine/Source/Runtime/Engine/Private/LevelStreaming.cpp
+++ b/Engine/Source/Runtime/Engine/Private/LevelStreaming.cpp
@@ -1717,10 +1717,9 @@ bool ULevelStreaming::RequestLevel(UWorld* PersistentWorld, bool bAllowLevelLoad
 		ULevel::StreamedLevelsOwningWorld.Add(DesiredPackageName, PersistentWorld);
 		UWorld::WorldTypePreLoadMap.FindOrAdd(DesiredPackageName) = PersistentWorld->WorldType;
 
-		// Kick off async load request.
-		STAT_ADD_CUSTOMMESSAGE_NAME( STAT_NamedMarker, *(FString( TEXT( "RequestLevel - " ) + DesiredPackageName.ToString() )) );
-		TRACE_BOOKMARK(TEXT("RequestLevel - %s"), *DesiredPackageName.ToString());
-			
+			// Kick off async load request.
+			STAT_ADD_CUSTOMMESSAGE_NAME( STAT_NamedMarker, *(FString( TEXT( "RequestLevel - " ) + DesiredPackageName.ToString() )) );
+			TRACE_BEGIN_REGION(*FString::Printf(TEXT("RequestLevel - %s"), *DesiredPackageName.ToString()))
 		FLinkerInstancingContext* InstancingContextPtr = nullptr;
 #if WITH_EDITOR
 		FLinkerInstancingContext InstancingContext;
@@ -1915,7 +1914,7 @@ void ULevelStreaming::AsyncLevelLoadComplete(const FName& InPackageName, UPackag
 	ULevel::StreamedLevelsOwningWorld.Remove(InPackageName);
 
 	STAT_ADD_CUSTOMMESSAGE_NAME( STAT_NamedMarker, *(FString( TEXT( "RequestLevelComplete - " ) + InPackageName.ToString() )) );
-	TRACE_BOOKMARK(TEXT("RequestLevelComplete - %s"), *InPackageName.ToString());
+	TRACE_END_REGION(*FString::Printf(TEXT("RequestLevel - %s"), *InPackageName.ToString()))
 }
 
 bool ULevelStreaming::IsStreamingStatePending() const
diff --git a/Engine/Source/Runtime/Engine/Private/World.cpp b/Engine/Source/Runtime/Engine/Private/World.cpp
--- a/Engine/Source/Runtime/Engine/Private/World.cpp
+++ b/Engine/Source/Runtime/Engine/Private/World.cpp
@@ -7948,7 +7948,7 @@ void FSeamlessTravelHandler::SeamlessTravelLoadCallback(const FName& PackageName
 	}
 
 	STAT_ADD_CUSTOMMESSAGE_NAME( STAT_NamedMarker, *(FString( TEXT( "StartTravelComplete - " ) + PackageName.ToString() )) );
-	TRACE_BOOKMARK(TEXT("StartTravelComplete - %s"), *PackageName.ToString());
+ 	TRACE_END_REGION(*FString::Printf(TEXT("SeamlessTravel - %s"), *PackageName.ToString()))
 }
 
 bool FSeamlessTravelHandler::StartTravel(UWorld* InCurrentWorld, const FURL& InURL)
@@ -8053,8 +8053,7 @@ bool FSeamlessTravelHandler::StartTravel(UWorld* InCurrentWorld, const FURL& InU
 				{
 					// Load the transition map, possibly with PIE prefix
 					STAT_ADD_CUSTOMMESSAGE_NAME( STAT_NamedMarker, *(FString( TEXT( "StartTravel - " ) + TransitionMap )) );
-					TRACE_BOOKMARK(TEXT("StartTravel - %s"), *TransitionMap);
-
+					TRACE_BEGIN_REGION(*FString::Printf(TEXT("SeamlessTravel - %s"), *TransitionMap))
 					if (!StartLoadingMap(TransitionMap))
 					{
 						UE_LOG(LogWorld, Error, TEXT("StartTravel: Invalid TransitionMap \"%s\""), *TransitionMap);
diff --git a/Engine/Source/Runtime/Online/HTTP/Private/HttpManager.cpp b/Engine/Source/Runtime/Online/HTTP/Private/HttpManager.cpp
--- a/Engine/Source/Runtime/Online/HTTP/Private/HttpManager.cpp
+++ b/Engine/Source/Runtime/Online/HTTP/Private/HttpManager.cpp
@@ -576,12 +576,16 @@ void FHttpManager::AddRequest(const FHttpRequestRef& Request)
 
 	QUICK_SCOPE_CYCLE_COUNTER(STAT_FHttpManager_RequestAddedDelegate);
 	RequestAddedDelegate.ExecuteIfBound(Request);
+
+#if MISCTRACE_ENABLED
+	auto Id = TRACE_BEGIN_REGION_WITH_ID(*FString::Printf(TEXT("HTTP: %s - %s"), *Request->GetVerb(), *Request->GetURL()), TEXT("HTTP"))
+	RequestRegionIds.Add(Request, Id);
+#endif
 }
 
 void FHttpManager::RemoveRequest(const FHttpRequestRef& Request)
 {
 	FScopeLock ScopeLock(&RequestLock);
-
 	Requests.Remove(Request);
 }
 
@@ -669,7 +673,12 @@ void FHttpManager::BroadcastHttpRequestCompleted(const FHttpRequestRef& Request)
 		// Increment index
 		HttpStatsHistory.HistoryIndex = (HttpStatsHistory.HistoryIndex + 1) % FHttpStatsHistory::HttpHistoryCount;
 	}
-		
+
+#if MISCTRACE_ENABLED
+	uint64 Id = RequestRegionIds.FindChecked(Request);
+	TRACE_END_REGION_WITH_ID(Id);
+	RequestRegionIds.Remove(Request);
+#endif
 	RequestCompletedDelegate.ExecuteIfBound(Request);
 }
 
diff --git a/Engine/Source/Runtime/Online/HTTP/Public/HttpManager.h b/Engine/Source/Runtime/Online/HTTP/Public/HttpManager.h
--- a/Engine/Source/Runtime/Online/HTTP/Public/HttpManager.h
+++ b/Engine/Source/Runtime/Online/HTTP/Public/HttpManager.h
@@ -17,6 +17,7 @@
 #include "Misc/URLRequestFilter.h"
 #include "Templates/Function.h"
 #include "Templates/SharedPointer.h"
+#include "ProfilingDebugging/MiscTrace.h"
 
 class FHttpThreadBase;
 class FHttpRequestCommon;
@@ -419,7 +420,11 @@ protected:
 	mutable FCriticalSection ThreadsCompletingRequestCriticalSection;
 	TMap<uint32, bool> ThreadsCompletingRequest;
 #endif
-
+	
+#if MISCTRACE_ENABLED
+	// This is storing Timing Region Ids per request for tracing HTTP requests
+	TMap<FHttpRequestRef, uint64> RequestRegionIds;
+#endif
 PACKAGE_SCOPE:
 
 	/** Used to lock access to add/remove/find requests */
diff --git a/Engine/Source/Runtime/Slate/Private/Framework/Notifications/NotificationManager.cpp b/Engine/Source/Runtime/Slate/Private/Framework/Notifications/NotificationManager.cpp
--- a/Engine/Source/Runtime/Slate/Private/Framework/Notifications/NotificationManager.cpp
+++ b/Engine/Source/Runtime/Slate/Private/Framework/Notifications/NotificationManager.cpp
@@ -233,7 +233,10 @@ FProgressNotificationHandle FSlateNotificationManager::StartProgressNotification
 	{
 		FProgressNotificationHandle NewHandle(++ProgressHandleCounter);
 		ProgressNotificationHandler->StartProgressNotification(NewHandle, DisplayText, TotalWorkToDo);
-
+#if MISCTRACE_ENABLED
+		uint64_t Id = TRACE_BEGIN_REGION_WITH_ID(*DisplayText.ToString(), TEXT("EditorNotifications"))
+		OpenRegions.Add(NewHandle.Id, Id);
+#endif
 		return NewHandle;
 	}
 
@@ -244,6 +247,24 @@ void FSlateNotificationManager::UpdateProgressNotification(FProgressNotification
 {
 	if (ProgressNotificationHandler)
 	{
+#if MISCTRACE_ENABLED
+		if (UpdatedTotalWorkToDo - TotalWorkDone == 0)
+		{
+			if (OpenRegions.Contains(InHandle.Id))
+			{
+				const uint64_t RegionID = OpenRegions.FindAndRemoveChecked(InHandle.Id);
+				TRACE_END_REGION_WITH_ID(RegionID);
+			}
+		}
+		else
+		{
+			if (!OpenRegions.Contains(InHandle.Id))
+			{
+				const uint64_t RegionID = TRACE_BEGIN_REGION_WITH_ID(*UpdatedDisplayText.ToString(), TEXT("EditorNotifications"))
+				OpenRegions.Add(InHandle.Id, RegionID);
+			}
+		}
+#endif
 		ProgressNotificationHandler->UpdateProgressNotification(InHandle, TotalWorkDone, UpdatedTotalWorkToDo, UpdatedDisplayText);
 	}
 }
@@ -253,6 +274,13 @@ void FSlateNotificationManager::CancelProgressNotification(FProgressNotification
 	if (ProgressNotificationHandler)
 	{
 		ProgressNotificationHandler->CancelProgressNotification(InHandle);
+#if MISCTRACE_ENABLED
+		if (OpenRegions.Contains(InHandle.Id))
+		{
+			uint64_t RegionID = OpenRegions.FindAndRemoveChecked(InHandle.Id);
+			TRACE_END_REGION_WITH_ID(RegionID);
+		}
+#endif
 	}
 }
 
diff --git a/Engine/Source/Runtime/Slate/Public/Framework/Notifications/NotificationManager.h b/Engine/Source/Runtime/Slate/Public/Framework/Notifications/NotificationManager.h
--- a/Engine/Source/Runtime/Slate/Public/Framework/Notifications/NotificationManager.h
+++ b/Engine/Source/Runtime/Slate/Public/Framework/Notifications/NotificationManager.h
@@ -10,6 +10,7 @@
 
 #include "HAL/IConsoleManager.h"
 #include "Misc/CoreAsyncTaskNotificationImpl.h"
+#include "ProfilingDebugging/MiscTrace.h"
 
 struct FNotificationInfo;
 
@@ -188,7 +189,11 @@ protected:
 	SLATE_API void ShutdownOnPreExit();
 
 private:
-
+#if MISCTRACE_ENABLED
+	// stores map of FProgressNotificationHandle.ID to open trace region ids. Used to visualize progress notifications in a trace file
+	TMap<int32, uint64_t> OpenRegions;
+#endif
+	
 	/** A list of notifications, bound to a particular region */
 	struct FRegionalNotificationList
 	{
