Replace some usages of FORCEINLINE with inline in Analytics modules.

For context about this change, see this doc: https://docs.google.com/document/d/1LjpT7NpzTecfXdFNyv5KM6sCu67827BRljNi1iaWVSE/edit?tab=t.0#heading=h.fvtz8gjukz7x

Includes modules: Analytics, AnalyticsVisualEditing

Manually reviewed.

#rnx

diff --git a/Engine/Source/Runtime/Analytics/Analytics/Public/AnalyticsEventAttribute.h b/Engine/Source/Runtime/Analytics/Analytics/Public/AnalyticsEventAttribute.h
--- a/Engine/Source/Runtime/Analytics/Analytics/Public/AnalyticsEventAttribute.h
+++ b/Engine/Source/Runtime/Analytics/Analytics/Public/AnalyticsEventAttribute.h
@@ -9,7 +9,7 @@ struct FJsonNull
 {
 };
 
-FORCEINLINE const TCHAR* LexToString(FJsonNull)
+inline const TCHAR* LexToString(FJsonNull)
 {
 	return TEXT("null");
 }
@@ -23,14 +23,14 @@ struct FJsonFragment
 	FString FragmentString;
 };
 
-FORCEINLINE const TCHAR* LexToString(const FJsonFragment& Fragment)
+inline const TCHAR* LexToString(const FJsonFragment& Fragment)
 {
 	if (Fragment.FragmentString.IsEmpty())
 		return LexToString(FJsonNull{});
 	return *Fragment.FragmentString;
 }
 
-FORCEINLINE FString LexToString(FJsonFragment&& Fragment)
+inline FString LexToString(FJsonFragment&& Fragment)
 {
 	if (Fragment.FragmentString.IsEmpty())
 		return LexToString(FJsonNull{});
@@ -251,20 +251,20 @@ namespace ImplMakeAnalyticsEventAttributeArray
 {
 	/** Recursion terminator. Empty list. */
 	template <typename Allocator>
-	FORCEINLINE void MakeArray(TArray<FAnalyticsEventAttribute, Allocator>& Attrs)
+	inline void MakeArray(TArray<FAnalyticsEventAttribute, Allocator>& Attrs)
 	{
 	}
 
 	/** Recursion terminator. Convert the key/value pair to analytics strings. */
 	template <typename Allocator, typename KeyType, typename ValueType>
-	FORCEINLINE void MakeArray(TArray<FAnalyticsEventAttribute, Allocator>& Attrs, KeyType&& Key, ValueType&& Value)
+	inline void MakeArray(TArray<FAnalyticsEventAttribute, Allocator>& Attrs, KeyType&& Key, ValueType&& Value)
 	{
 		Attrs.Emplace(Forward<KeyType>(Key), Forward<ValueType>(Value));
 	}
 
 	/** recursively add the arguments to the array. */
 	template <typename Allocator, typename KeyType, typename ValueType, typename...ArgTypes>
-	FORCEINLINE void MakeArray(TArray<FAnalyticsEventAttribute, Allocator>& Attrs, KeyType&& Key, ValueType&& Value, ArgTypes&&...Args)
+	inline void MakeArray(TArray<FAnalyticsEventAttribute, Allocator>& Attrs, KeyType&& Key, ValueType&& Value, ArgTypes&&...Args)
 	{
 		// pop off the top two args and recursively apply the rest.
 		Attrs.Emplace(Forward<KeyType>(Key), Forward<ValueType>(Value));
@@ -274,7 +274,7 @@ namespace ImplMakeAnalyticsEventAttributeArray
 
 /** Helper to create an array of attributes using a single expression. Reserves the necessary space in advance. There must be an even number of arguments, one for each key/value pair. */
 template <typename Allocator = FDefaultAllocator, typename...ArgTypes>
-FORCEINLINE TArray<FAnalyticsEventAttribute, Allocator> MakeAnalyticsEventAttributeArray(ArgTypes&&...Args)
+inline TArray<FAnalyticsEventAttribute, Allocator> MakeAnalyticsEventAttributeArray(ArgTypes&&...Args)
 {
 	static_assert(sizeof...(Args) % 2 == 0, "Must pass an even number of arguments.");
 	TArray<FAnalyticsEventAttribute, Allocator> Attrs;
@@ -285,7 +285,7 @@ FORCEINLINE TArray<FAnalyticsEventAttribute, Allocator> MakeAnalyticsEventAttrib
 
 /** Helper to append to an array of attributes using a single expression. Reserves the necessary space in advance. There must be an even number of arguments, one for each key/value pair. */
 template <typename Allocator = FDefaultAllocator, typename...ArgTypes>
-FORCEINLINE TArray<FAnalyticsEventAttribute, Allocator>& AppendAnalyticsEventAttributeArray(TArray<FAnalyticsEventAttribute, Allocator>& Attrs, ArgTypes&&...Args)
+inline TArray<FAnalyticsEventAttribute, Allocator>& AppendAnalyticsEventAttributeArray(TArray<FAnalyticsEventAttribute, Allocator>& Attrs, ArgTypes&&...Args)
 {
 	static_assert(sizeof...(Args) % 2 == 0, "Must pass an even number of arguments.");
 	Attrs.Reserve(Attrs.Num() + (sizeof...(Args) / 2));
diff --git a/Engine/Source/Runtime/Analytics/AnalyticsVisualEditing/Classes/AnalyticsSettings.h b/Engine/Source/Runtime/Analytics/AnalyticsVisualEditing/Classes/AnalyticsSettings.h
--- a/Engine/Source/Runtime/Analytics/AnalyticsVisualEditing/Classes/AnalyticsSettings.h
+++ b/Engine/Source/Runtime/Analytics/AnalyticsVisualEditing/Classes/AnalyticsSettings.h
@@ -24,11 +24,11 @@ class ANALYTICSVISUALEDITING_API UAnalyticsSettingsBase
 	FText SettingsTooltip;
 
 	/** Helpers so that subclasses don't need to hardcode the strings in case they change */
-	FORCEINLINE FString GetReleaseIniSection() const { return TEXT("Analytics"); }
-	FORCEINLINE FString GetDebugIniSection() const { return TEXT("AnalyticsDebug"); }
-	FORCEINLINE FString GetTestIniSection() const { return TEXT("AnalyticsTest"); }
-	FORCEINLINE FString GetDevelopmentIniSection() const { return TEXT("AnalyticsDevelopment"); }
-	FORCEINLINE FString GetIniName() const { return FString::Printf(TEXT("%sDefaultEngine.ini"), *FPaths::SourceConfigDir()); }
+	inline FString GetReleaseIniSection() const { return TEXT("Analytics"); }
+	inline FString GetDebugIniSection() const { return TEXT("AnalyticsDebug"); }
+	inline FString GetTestIniSection() const { return TEXT("AnalyticsTest"); }
+	inline FString GetDevelopmentIniSection() const { return TEXT("AnalyticsDevelopment"); }
+	inline FString GetIniName() const { return FString::Printf(TEXT("%sDefaultEngine.ini"), *FPaths::SourceConfigDir()); }
 
 private:
 /**
