Removed slow default attributes to event attributes compare in IAnalyticsProvider

diff --git a/Engine/Source/Runtime/Analytics/Analytics/Private/Interfaces/IAnalyticsProvider.cpp b/Engine/Source/Runtime/Analytics/Analytics/Private/Interfaces/IAnalyticsProvider.cpp
--- a/Engine/Source/Runtime/Analytics/Analytics/Private/Interfaces/IAnalyticsProvider.cpp
+++ b/Engine/Source/Runtime/Analytics/Analytics/Private/Interfaces/IAnalyticsProvider.cpp
@@ -21,22 +21,3 @@ void IAnalyticsProvider::CheckForDuplicateAttributes(const FString& EventName, c
 	}
 #endif
 }
-
-void IAnalyticsProvider::CheckForDuplicateAttributes(const FString& EventName, const TArray<FAnalyticsEventAttribute>& AttributesA, const TArray<FAnalyticsEventAttribute>& AttributesB)
-{
-#if UE_BUILD_DEBUG || UE_BUILD_DEVELOPMENT
-	// Expose events that have duplicate attribute names. Not only does this provide ambiguous event data, in general this is not handled by all the analytics backends in any consistent manner.
-	for (int32 indexA = 0; indexA < AttributesA.Num(); ++indexA)
-	{
-		const FString& AttributeName = AttributesA[indexA].GetName();
-
-		for (int32 indexB = 0; indexB < AttributesB.Num(); ++indexB)
-		{
-			if (AttributeName == AttributesB[indexB].GetName())
-			{
-				UE_LOG(LogAnalytics, Warning, TEXT("Duplicate Attributes %s Found For Event %s"), *AttributeName, *EventName);
-			}
-		}
-	}
-#endif
-}
diff --git a/Engine/Source/Runtime/Analytics/Analytics/Public/Interfaces/IAnalyticsProvider.h b/Engine/Source/Runtime/Analytics/Analytics/Public/Interfaces/IAnalyticsProvider.h
--- a/Engine/Source/Runtime/Analytics/Analytics/Public/Interfaces/IAnalyticsProvider.h
+++ b/Engine/Source/Runtime/Analytics/Analytics/Public/Interfaces/IAnalyticsProvider.h
@@ -451,11 +451,5 @@ protected:
 	 * Checks for duplicate key names in the same attribute array. Only implemented for debug and development builds.
 	 */
 	void ANALYTICS_API CheckForDuplicateAttributes(const FString& EventName, const TArray<FAnalyticsEventAttribute>& Attributes);
-
-	/**
-	 * Checks for duplicate key names between two different attribute arrays. Only implemented for debug and development builds.
-	 */
-	void ANALYTICS_API CheckForDuplicateAttributes(const FString& EventName, const TArray<FAnalyticsEventAttribute>& AttributesA, const TArray<FAnalyticsEventAttribute>& AttributesB);
-
 };
 
diff --git a/Engine/Source/Runtime/Analytics/AnalyticsET/Private/IAnalyticsProviderET.cpp b/Engine/Source/Runtime/Analytics/AnalyticsET/Private/IAnalyticsProviderET.cpp
--- a/Engine/Source/Runtime/Analytics/AnalyticsET/Private/IAnalyticsProviderET.cpp
+++ b/Engine/Source/Runtime/Analytics/AnalyticsET/Private/IAnalyticsProviderET.cpp
@@ -622,9 +622,6 @@ void FAnalyticsProviderET::RecordEvent(FString&& EventName, const TArray<FAnalyt
 	// Check for duplicate attributes against themselves
 	CheckForDuplicateAttributes(EventName, Attributes);
 
-	// Check for duplicate attributes against the defaults
-	CheckForDuplicateAttributes(EventName, Attributes, GetDefaultEventAttributesSafe());
-
 	// let higher level code filter the decision of whether to send the event
 	if (ShouldRecordEvent(EventName))
 	{
