Use a C++17 fold operation to avoid doing a recursive template and shave 2.26MB from a large server binary.

diff --git a/Engine/Source/Runtime/Analytics/Analytics/Public/AnalyticsEventAttribute.h b/Engine/Source/Runtime/Analytics/Analytics/Public/AnalyticsEventAttribute.h
--- a/Engine/Source/Runtime/Analytics/Analytics/Public/AnalyticsEventAttribute.h
+++ b/Engine/Source/Runtime/Analytics/Analytics/Public/AnalyticsEventAttribute.h
@@ -245,30 +245,47 @@ inline void FAnalyticsEventAttribute::SwitchToJsonFragment()
 
 PRAGMA_ENABLE_DEPRECATION_WARNINGS
 
-
 /** Helper functions for MakeAnalyticsEventAttributeArray. */
 namespace ImplMakeAnalyticsEventAttributeArray
 {
-	/** Recursion terminator. Empty list. */
 	template <typename Allocator>
-	inline void MakeArray(TArray<FAnalyticsEventAttribute, Allocator>& Attrs)
-	{
-	}
-
-	/** Recursion terminator. Convert the key/value pair to analytics strings. */
-	template <typename Allocator, typename KeyType, typename ValueType>
-	inline void MakeArray(TArray<FAnalyticsEventAttribute, Allocator>& Attrs, KeyType&& Key, ValueType&& Value)
+	struct Helper final
 	{
-		Attrs.Emplace(Forward<KeyType>(Key), Forward<ValueType>(Value));
-	}
+		Helper(TArray<FAnalyticsEventAttribute, Allocator>& Attrs) : Attrs(Attrs) {}
+
+		TOptional<FString> LastKey;
+		TArray<FAnalyticsEventAttribute, Allocator>& Attrs;
+
+		template <typename Type>
+		void ForEach(Type&& Thing)
+		{
+			if (!LastKey)
+			{
+				// Even indices are keys.
+				if constexpr (std::is_convertible_v<Type, FString>)
+				{
+					LastKey = ::MoveTempIfPossible(Thing);
+				}
+				else
+				{
+					checkf(false, TEXT("Even indices are keys and must be convertible to `FString`!"));
+				}
+			}
+			else
+			{
+				// Odd indices are values.
+				Attrs.Emplace(::MoveTemp(*LastKey), Thing);
+				LastKey.Reset();
+			}
+		}
+	};
 
-	/** recursively add the arguments to the array. */
-	template <typename Allocator, typename KeyType, typename ValueType, typename...ArgTypes>
-	inline void MakeArray(TArray<FAnalyticsEventAttribute, Allocator>& Attrs, KeyType&& Key, ValueType&& Value, ArgTypes&&...Args)
+	template <typename Allocator, typename...ArgTypes>
+	inline void MakeArray(TArray<FAnalyticsEventAttribute, Allocator>& Attrs, ArgTypes&&...Args)
 	{
-		// pop off the top two args and recursively apply the rest.
-		Attrs.Emplace(Forward<KeyType>(Key), Forward<ValueType>(Value));
-		MakeArray(Attrs, Forward<ArgTypes>(Args)...);
+		static_assert(sizeof...(Args) % 2 == 0, "Must pass an even number of arguments.");
+		Helper<Allocator> H(Attrs);
+		(H.ForEach(Args), ...);
 	}
 }
 
@@ -276,7 +293,6 @@ namespace ImplMakeAnalyticsEventAttributeArray
 template <typename Allocator = FDefaultAllocator, typename...ArgTypes>
 inline TArray<FAnalyticsEventAttribute, Allocator> MakeAnalyticsEventAttributeArray(ArgTypes&&...Args)
 {
-	static_assert(sizeof...(Args) % 2 == 0, "Must pass an even number of arguments.");
 	TArray<FAnalyticsEventAttribute, Allocator> Attrs;
 	Attrs.Empty(sizeof...(Args) / 2);
 	ImplMakeAnalyticsEventAttributeArray::MakeArray(Attrs, Forward<ArgTypes>(Args)...);
@@ -287,7 +303,6 @@ inline TArray<FAnalyticsEventAttribute, Allocator> MakeAnalyticsEventAttributeAr
 template <typename Allocator = FDefaultAllocator, typename...ArgTypes>
 inline TArray<FAnalyticsEventAttribute, Allocator>& AppendAnalyticsEventAttributeArray(TArray<FAnalyticsEventAttribute, Allocator>& Attrs, ArgTypes&&...Args)
 {
-	static_assert(sizeof...(Args) % 2 == 0, "Must pass an even number of arguments.");
 	Attrs.Reserve(Attrs.Num() + (sizeof...(Args) / 2));
 	ImplMakeAnalyticsEventAttributeArray::MakeArray(Attrs, Forward<ArgTypes>(Args)...);
 	return Attrs;
