[AutoRTFM] Update AnalyticsPropertyStore for AutoRTFM transactional safety.

This is reachable from EngineTestAutoRTFM when running tests from the automation window.

We cannot Flush from inside a transaction, so flushes are queued up for later. If the property store is destroyed inside a transaction, flushing won't happen, so this will cause an ensure-failure. (As long as we don't have short-lived/temporary property stores being created, this should be fine.)

This class also happens to mark the first time we've seen TryLock called on a transactionally-safe critical section. I've added a stub implementation for now which always returns false while inside a transaction; this isn't ideal, but I've cut a bug to write a proper implementation.

diff --git a/Engine/Source/Runtime/Analytics/Analytics/Private/AnalyticsPropertyStore.cpp b/Engine/Source/Runtime/Analytics/Analytics/Private/AnalyticsPropertyStore.cpp
--- a/Engine/Source/Runtime/Analytics/Analytics/Private/AnalyticsPropertyStore.cpp
+++ b/Engine/Source/Runtime/Analytics/Analytics/Private/AnalyticsPropertyStore.cpp
@@ -1,6 +1,7 @@
 // Copyright Epic Games, Inc. All Rights Reserved.
 
 #include "AnalyticsPropertyStore.h"
+#include "AutoRTFM.h"
 #include "Containers/StringConv.h"
 #include "GenericPlatform/GenericPlatformFile.h"
 #include "HAL/FileManager.h"
@@ -14,7 +15,7 @@
 
 namespace AnalyticsPropertyStoreUtils
 {
-	/** The key/value storage format version defininig the memory layout of the store. */
+	/** The key/value storage format version defining the memory layout of the store. */
 	static const uint32 StoreFormatVersion = 1;
 
 	/** MD5 checksum is 16 bytes long. */
@@ -28,7 +29,7 @@ namespace AnalyticsPropertyStoreUtils
 		uint8 Checksum[Md5DigestSize];
 	};
 
-	/** Conditionnaly set a value is the store. */
+	/** Conditionally set a value in the store. */
 	template<typename T, typename TGetter, typename TSetter, typename TCompare>
 	IAnalyticsPropertyStore::EStatusCode SetValue(TGetter&& GetterFn, TSetter&& SetterFn, const T& ProposedValue, TCompare&& ConditionFn)
 	{
@@ -49,7 +50,7 @@ namespace AnalyticsPropertyStoreUtils
 		return Status;
 	}
 
-	/** Conditionnaly update a value is the store. */
+	/** Conditionally update a value in the store. */
 	template<typename T, typename TGetter, typename TSetter, typename TCompare>
 	IAnalyticsPropertyStore::EStatusCode UpdateValue(TGetter&& GetterFn, TSetter&& SetterFn, TCompare&& UpdateFn)
 	{
@@ -80,13 +81,22 @@ FAnalyticsPropertyStore::FAnalyticsPropertyStore()
 
 FAnalyticsPropertyStore::~FAnalyticsPropertyStore()
 {
-	FScopeLock Lock(&StoreLock);
-
-	// Ensure any running async flush has completed.
-	if (FlushTask.IsValid())
+	if (AutoRTFM::IsClosed())
 	{
-		FlushTask->EnsureCompletion();
-		FlushTask.Reset();
+		// The FAnalyticsPropertyStore must not be destroyed from within a transaction, 
+		// since any queued-up flush tasks won't be run. 
+		AutoRTFM::PopAllOnCommitHandlers(this);
+		ensure(!"Destroying FAnalyticsPropertyStore from within a transaction!");
+	}
+	else
+	{
+		// Ensure any running async flush has completed.
+		UE::TScopeLock Lock(StoreLock);
+		if (FlushTask.IsValid())
+		{
+			FlushTask->EnsureCompletion();
+			FlushTask.Reset();
+		}
 	}
 }
 
@@ -108,7 +118,7 @@ void FAnalyticsPropertyStore::Reset()
 
 bool FAnalyticsPropertyStore::Create(const FString& Pathname, uint32 CapacityHint)
 {
-	FScopeLock Lock(&StoreLock);
+	UE::TScopeLock Lock(StoreLock);
 
 	// Reset internal state.
 	Reset();
@@ -132,7 +142,7 @@ bool FAnalyticsPropertyStore::Create(const FString& Pathname, uint32 CapacityHin
 
 bool FAnalyticsPropertyStore::Load(const FString& Pathname)
 {
-	FScopeLock Lock(&StoreLock);
+	UE::TScopeLock Lock(StoreLock);
 
 	// Reset internal state.
 	Reset();
@@ -275,19 +285,19 @@ bool FAnalyticsPropertyStore::Load(const FString& Pathname)
 
 uint32 FAnalyticsPropertyStore::Num() const
 {
-	FScopeLock Lock(&StoreLock);
+	UE::TScopeLock Lock(StoreLock);
 	return NameOffsetMap.Num();
 }
 
 bool FAnalyticsPropertyStore::Contains(const FString& Key) const
 {
-	FScopeLock Lock(&StoreLock);
+	UE::TScopeLock Lock(StoreLock);
 	return NameOffsetMap.Contains(Key);
 }
 
 bool FAnalyticsPropertyStore::Remove(const FString& Key)
 {
-	FScopeLock Lock(&StoreLock);
+	UE::TScopeLock Lock(StoreLock);
 	uint32 Offset;
 	if (NameOffsetMap.RemoveAndCopyValue(Key, Offset))
 	{
@@ -306,7 +316,7 @@ bool FAnalyticsPropertyStore::Remove(const FString& Key)
 
 void FAnalyticsPropertyStore::RemoveAll()
 {
-	FScopeLock Lock(&StoreLock);
+	UE::TScopeLock Lock(StoreLock);
 	NameOffsetMap.Reset();
 	StorageReader.Seek(0);
 	StorageWriter.Seek(0);
@@ -316,61 +326,61 @@ void FAnalyticsPropertyStore::RemoveAll()
 
 IAnalyticsPropertyStore::EStatusCode FAnalyticsPropertyStore::Set(const FString& Key, int32 Value)
 {
-	FScopeLock Lock(&StoreLock);
+	UE::TScopeLock Lock(StoreLock);
 	return SetFixedSizeValueInternal(Key, ETypeCode::I32, Value);
 }
 
 IAnalyticsPropertyStore::EStatusCode FAnalyticsPropertyStore::Set(const FString& Key, uint32 Value)
 {
-	FScopeLock Lock(&StoreLock);
+	UE::TScopeLock Lock(StoreLock);
 	return SetFixedSizeValueInternal(Key, ETypeCode::U32, Value);
 }
 
 IAnalyticsPropertyStore::EStatusCode FAnalyticsPropertyStore::Set(const FString& Key, int64 Value)
 {
-	FScopeLock Lock(&StoreLock);
+	UE::TScopeLock Lock(StoreLock);
 	return SetFixedSizeValueInternal(Key, ETypeCode::I64, Value);
 }
 
 IAnalyticsPropertyStore::EStatusCode FAnalyticsPropertyStore::Set(const FString& Key, uint64 Value)
 {
-	FScopeLock Lock(&StoreLock);
+	UE::TScopeLock Lock(StoreLock);
 	return SetFixedSizeValueInternal(Key, ETypeCode::U64, Value);
 }
 
 IAnalyticsPropertyStore::EStatusCode FAnalyticsPropertyStore::Set(const FString& Key, float Value)
 {
-	FScopeLock Lock(&StoreLock);
+	UE::TScopeLock Lock(StoreLock);
 	return SetFixedSizeValueInternal(Key, ETypeCode::Flt, Value);
 }
 
 IAnalyticsPropertyStore::EStatusCode FAnalyticsPropertyStore::Set(const FString& Key, double Value)
 {
-	FScopeLock Lock(&StoreLock);
+	UE::TScopeLock Lock(StoreLock);
 	return SetFixedSizeValueInternal(Key, ETypeCode::Dbl, Value);
 }
 
 IAnalyticsPropertyStore::EStatusCode FAnalyticsPropertyStore::Set(const FString& Key, bool Value)
 {
-	FScopeLock Lock(&StoreLock);
+	UE::TScopeLock Lock(StoreLock);
 	return SetFixedSizeValueInternal(Key, ETypeCode::Bool, Value);
 }
 
 IAnalyticsPropertyStore::EStatusCode FAnalyticsPropertyStore::Set(const FString& Key, const FString& Value, uint32 CharCountCapacityHint)
 {
-	FScopeLock Lock(&StoreLock);
+	UE::TScopeLock Lock(StoreLock);
 	return SetStringValueInternal(Key, Value, CharCountCapacityHint);
 }
 
 IAnalyticsPropertyStore::EStatusCode FAnalyticsPropertyStore::Set(const FString& Key, const FDateTime& Value)
 {
-	FScopeLock Lock(&StoreLock);
+	UE::TScopeLock Lock(StoreLock);
 	return SetDateTimeValueInternal(Key, Value);
 }
 
 IAnalyticsPropertyStore::EStatusCode FAnalyticsPropertyStore::Set(const FString& Key, int32 Value, const TFunction<bool(const int32*, const int32&)>& ConditionFn)
 {
-	FScopeLock Lock(&StoreLock);
+	UE::TScopeLock Lock(StoreLock);
 
 	return AnalyticsPropertyStoreUtils::SetValue(
 				[this, &Key](int32& OutValue) { return GetFixedSizeValueInternal(Key, ETypeCode::I32, OutValue); },
@@ -381,7 +391,7 @@ IAnalyticsPropertyStore::EStatusCode FAnalyticsPropertyStore::Set(const FString&
 
 IAnalyticsPropertyStore::EStatusCode FAnalyticsPropertyStore::Set(const FString& Key, uint32 Value, const TFunction<bool(const uint32*, const uint32&)>& ConditionFn)
 {
-	FScopeLock Lock(&StoreLock);
+	UE::TScopeLock Lock(StoreLock);
 
 	return AnalyticsPropertyStoreUtils::SetValue(
 				[this, &Key](uint32& OutValue) { return GetFixedSizeValueInternal(Key, ETypeCode::U32, OutValue); },
@@ -392,7 +402,7 @@ IAnalyticsPropertyStore::EStatusCode FAnalyticsPropertyStore::Set(const FString&
 
 IAnalyticsPropertyStore::EStatusCode FAnalyticsPropertyStore::Set(const FString& Key, int64 Value, const TFunction<bool(const int64*, const int64&)>& ConditionFn)
 {
-	FScopeLock Lock(&StoreLock);
+	UE::TScopeLock Lock(StoreLock);
 
 	return AnalyticsPropertyStoreUtils::SetValue(
 				[this, &Key](int64& OutValue) { return GetFixedSizeValueInternal(Key, ETypeCode::I64, OutValue); },
@@ -403,7 +413,7 @@ IAnalyticsPropertyStore::EStatusCode FAnalyticsPropertyStore::Set(const FString&
 
 IAnalyticsPropertyStore::EStatusCode FAnalyticsPropertyStore::Set(const FString& Key, uint64 Value, const TFunction<bool(const uint64*, const uint64&)>& ConditionFn)
 {
-	FScopeLock Lock(&StoreLock);
+	UE::TScopeLock Lock(StoreLock);
 
 	return AnalyticsPropertyStoreUtils::SetValue(
 				[this, &Key](uint64& OutValue) { return GetFixedSizeValueInternal(Key, ETypeCode::U64, OutValue); },
@@ -414,7 +424,7 @@ IAnalyticsPropertyStore::EStatusCode FAnalyticsPropertyStore::Set(const FString&
 
 IAnalyticsPropertyStore::EStatusCode FAnalyticsPropertyStore::Set(const FString& Key, float Value, const TFunction<bool(const float*, const float&)>& ConditionFn)
 {
-	FScopeLock Lock(&StoreLock);
+	UE::TScopeLock Lock(StoreLock);
 
 	return AnalyticsPropertyStoreUtils::SetValue(
 				[this, &Key](float& OutValue) { return GetFixedSizeValueInternal(Key, ETypeCode::Flt, OutValue); },
@@ -425,7 +435,7 @@ IAnalyticsPropertyStore::EStatusCode FAnalyticsPropertyStore::Set(const FString&
 
 IAnalyticsPropertyStore::EStatusCode FAnalyticsPropertyStore::Set(const FString& Key, double Value, const TFunction<bool(const double*, const double&)>& ConditionFn)
 {
-	FScopeLock Lock(&StoreLock);
+	UE::TScopeLock Lock(StoreLock);
 
 	return AnalyticsPropertyStoreUtils::SetValue(
 				[this, &Key](double& OutValue) { return GetFixedSizeValueInternal(Key, ETypeCode::Dbl, OutValue); },
@@ -436,7 +446,7 @@ IAnalyticsPropertyStore::EStatusCode FAnalyticsPropertyStore::Set(const FString&
 
 IAnalyticsPropertyStore::EStatusCode FAnalyticsPropertyStore::Set(const FString& Key, bool Value, const TFunction<bool(const bool*, const bool&)>& ConditionFn)
 {
-	FScopeLock Lock(&StoreLock);
+	UE::TScopeLock Lock(StoreLock);
 
 	return AnalyticsPropertyStoreUtils::SetValue(
 				[this, &Key](bool& OutValue) { return GetFixedSizeValueInternal(Key, ETypeCode::Bool, OutValue); },
@@ -447,7 +457,7 @@ IAnalyticsPropertyStore::EStatusCode FAnalyticsPropertyStore::Set(const FString&
 
 IAnalyticsPropertyStore::EStatusCode FAnalyticsPropertyStore::Set(const FString& Key, const FString& Value, const TFunction<bool(const FString*, const FString&)>& ConditionFn)
 {
-	FScopeLock Lock(&StoreLock);
+	UE::TScopeLock Lock(StoreLock);
 
 	return AnalyticsPropertyStoreUtils::SetValue(
 				[this, &Key](FString& OutValue)      { return GetStringValueInternal(Key, OutValue); },
@@ -458,7 +468,7 @@ IAnalyticsPropertyStore::EStatusCode FAnalyticsPropertyStore::Set(const FString&
 
 IAnalyticsPropertyStore::EStatusCode FAnalyticsPropertyStore::Set(const FString& Key, const FDateTime& Value, const TFunction<bool(const FDateTime*, const FDateTime&)>& ConditionFn)
 {
-	FScopeLock Lock(&StoreLock);
+	UE::TScopeLock Lock(StoreLock);
 
 	return AnalyticsPropertyStoreUtils::SetValue(
 				[this, &Key](FDateTime& OutValue)      { return GetDateTimeValueInternal(Key, OutValue); },
@@ -469,7 +479,7 @@ IAnalyticsPropertyStore::EStatusCode FAnalyticsPropertyStore::Set(const FString&
 
 IAnalyticsPropertyStore::EStatusCode FAnalyticsPropertyStore::Update(const FString& Key, const TFunction<bool(int32&)>& UpdateFn)
 {
-	FScopeLock Lock(&StoreLock);
+	UE::TScopeLock Lock(StoreLock);
 
 	return AnalyticsPropertyStoreUtils::UpdateValue<int32>(
 				[this, &Key](int32& OutValue) { return GetFixedSizeValueInternal(Key, ETypeCode::I32, OutValue); },
@@ -479,7 +489,7 @@ IAnalyticsPropertyStore::EStatusCode FAnalyticsPropertyStore::Update(const FStri
 
 IAnalyticsPropertyStore::EStatusCode FAnalyticsPropertyStore::Update(const FString& Key, const TFunction<bool(uint32&)>& UpdateFn)
 {
-	FScopeLock Lock(&StoreLock);
+	UE::TScopeLock Lock(StoreLock);
 
 	return AnalyticsPropertyStoreUtils::UpdateValue<uint32>(
 				[this, &Key](uint32& OutValue) { return GetFixedSizeValueInternal(Key, ETypeCode::U32, OutValue); },
@@ -489,7 +499,7 @@ IAnalyticsPropertyStore::EStatusCode FAnalyticsPropertyStore::Update(const FStri
 
 IAnalyticsPropertyStore::EStatusCode FAnalyticsPropertyStore::Update(const FString& Key, const TFunction<bool(int64&)>& UpdateFn)
 {
-	FScopeLock Lock(&StoreLock);
+	UE::TScopeLock Lock(StoreLock);
 
 	return AnalyticsPropertyStoreUtils::UpdateValue<int64>(
 				[this, &Key](int64& OutValue) { return GetFixedSizeValueInternal(Key, ETypeCode::I64, OutValue); },
@@ -499,7 +509,7 @@ IAnalyticsPropertyStore::EStatusCode FAnalyticsPropertyStore::Update(const FStri
 
 IAnalyticsPropertyStore::EStatusCode FAnalyticsPropertyStore::Update(const FString& Key, const TFunction<bool(uint64&)>& UpdateFn)
 {
-	FScopeLock Lock(&StoreLock);
+	UE::TScopeLock Lock(StoreLock);
 
 	return AnalyticsPropertyStoreUtils::UpdateValue<uint64>(
 				[this, &Key](uint64& OutValue) { return GetFixedSizeValueInternal(Key, ETypeCode::U64, OutValue); },
@@ -509,7 +519,7 @@ IAnalyticsPropertyStore::EStatusCode FAnalyticsPropertyStore::Update(const FStri
 
 IAnalyticsPropertyStore::EStatusCode FAnalyticsPropertyStore::Update(const FString& Key, const TFunction<bool(float&)>& UpdateFn)
 {
-	FScopeLock Lock(&StoreLock);
+	UE::TScopeLock Lock(StoreLock);
 
 	return AnalyticsPropertyStoreUtils::UpdateValue<float>(
 				[this, &Key](float& OutValue) { return GetFixedSizeValueInternal(Key, ETypeCode::Flt, OutValue); },
@@ -519,7 +529,7 @@ IAnalyticsPropertyStore::EStatusCode FAnalyticsPropertyStore::Update(const FStri
 
 IAnalyticsPropertyStore::EStatusCode FAnalyticsPropertyStore::Update(const FString& Key, const TFunction<bool(double&)>& UpdateFn)
 {
-	FScopeLock Lock(&StoreLock);
+	UE::TScopeLock Lock(StoreLock);
 
 	return AnalyticsPropertyStoreUtils::UpdateValue<double>(
 				[this, &Key](double& OutValue) { return GetFixedSizeValueInternal(Key, ETypeCode::Dbl, OutValue); },
@@ -529,7 +539,7 @@ IAnalyticsPropertyStore::EStatusCode FAnalyticsPropertyStore::Update(const FStri
 
 IAnalyticsPropertyStore::EStatusCode FAnalyticsPropertyStore::Update(const FString& Key, const TFunction<bool(bool&)>& UpdateFn)
 {
-	FScopeLock Lock(&StoreLock);
+	UE::TScopeLock Lock(StoreLock);
 
 	return AnalyticsPropertyStoreUtils::UpdateValue<bool>(
 				[this, &Key](bool& OutValue) { return GetFixedSizeValueInternal(Key, ETypeCode::Bool, OutValue); },
@@ -539,7 +549,7 @@ IAnalyticsPropertyStore::EStatusCode FAnalyticsPropertyStore::Update(const FStri
 
 IAnalyticsPropertyStore::EStatusCode FAnalyticsPropertyStore::Update(const FString& Key, const TFunction<bool(FString&)>& UpdateFn)
 {
-	FScopeLock Lock(&StoreLock);
+	UE::TScopeLock Lock(StoreLock);
 
 	return AnalyticsPropertyStoreUtils::UpdateValue<FString>(
 				[this, &Key](FString& OutValue)      { return GetStringValueInternal(Key, OutValue); },
@@ -549,7 +559,7 @@ IAnalyticsPropertyStore::EStatusCode FAnalyticsPropertyStore::Update(const FStri
 
 IAnalyticsPropertyStore::EStatusCode FAnalyticsPropertyStore::Update(const FString& Key, const TFunction<bool(FDateTime&)>& UpdateFn)
 {
-	FScopeLock Lock(&StoreLock);
+	UE::TScopeLock Lock(StoreLock);
 
 	return AnalyticsPropertyStoreUtils::UpdateValue<FDateTime>(
 				[this, &Key](FDateTime& OutValue)      { return GetDateTimeValueInternal(Key, OutValue); },
@@ -559,55 +569,55 @@ IAnalyticsPropertyStore::EStatusCode FAnalyticsPropertyStore::Update(const FStri
 
 IAnalyticsPropertyStore::EStatusCode FAnalyticsPropertyStore::Get(const FString& Key, int32& OutValue) const
 {
-	FScopeLock Lock(&StoreLock);
+	UE::TScopeLock Lock(StoreLock);
 	return GetFixedSizeValueInternal(Key, ETypeCode::I32, OutValue);
 }
 
 IAnalyticsPropertyStore::EStatusCode FAnalyticsPropertyStore::Get(const FString& Key, uint32& OutValue) const
 {
-	FScopeLock Lock(&StoreLock);
+	UE::TScopeLock Lock(StoreLock);
 	return GetFixedSizeValueInternal(Key, ETypeCode::U32, OutValue);
 }
 
 IAnalyticsPropertyStore::EStatusCode FAnalyticsPropertyStore::Get(const FString& Key, int64& OutValue) const
 {
-	FScopeLock Lock(&StoreLock);
+	UE::TScopeLock Lock(StoreLock);
 	return GetFixedSizeValueInternal(Key, ETypeCode::I64, OutValue);
 }
 
 IAnalyticsPropertyStore::EStatusCode FAnalyticsPropertyStore::Get(const FString& Key, uint64& OutValue) const
 {
-	FScopeLock Lock(&StoreLock);
+	UE::TScopeLock Lock(StoreLock);
 	return GetFixedSizeValueInternal(Key, ETypeCode::U64, OutValue);
 }
 
 IAnalyticsPropertyStore::EStatusCode FAnalyticsPropertyStore::Get(const FString& Key, float& OutValue) const
 {
-	FScopeLock Lock(&StoreLock);
+	UE::TScopeLock Lock(StoreLock);
 	return GetFixedSizeValueInternal(Key, ETypeCode::Flt, OutValue);
 }
 
 IAnalyticsPropertyStore::EStatusCode FAnalyticsPropertyStore::Get(const FString& Key, double& OutValue) const
 {
-	FScopeLock Lock(&StoreLock);
+	UE::TScopeLock Lock(StoreLock);
 	return GetFixedSizeValueInternal(Key, ETypeCode::Dbl, OutValue);
 }
 
 IAnalyticsPropertyStore::EStatusCode FAnalyticsPropertyStore::Get(const FString& Key, bool& OutValue) const
 {
-	FScopeLock Lock(&StoreLock);
+	UE::TScopeLock Lock(StoreLock);
 	return GetFixedSizeValueInternal(Key, ETypeCode::Bool, OutValue);
 }
 
 IAnalyticsPropertyStore::EStatusCode FAnalyticsPropertyStore::Get(const FString& Key, FString& OutValue) const
 {
-	FScopeLock Lock(&StoreLock);
+	UE::TScopeLock Lock(StoreLock);
 	return GetStringValueInternal(Key, OutValue);
 }
 
 IAnalyticsPropertyStore::EStatusCode FAnalyticsPropertyStore::Get(const FString& Key, FDateTime& OutValue) const
 {
-	FScopeLock Lock(&StoreLock);
+	UE::TScopeLock Lock(StoreLock);
 	return GetDateTimeValueInternal(Key, OutValue);
 }
 
@@ -823,7 +833,7 @@ void FAnalyticsPropertyStore::Defragment()
 
 void FAnalyticsPropertyStore::VisitAll(const TFunction<void(FAnalyticsEventAttribute&&)>& VisitFn) const
 {
-	FScopeLock Lock(&StoreLock);
+	UE::TScopeLock Lock(StoreLock);
 
 	for (const TPair<FString, uint32>& NameOffsetPair : NameOffsetMap)
 	{
@@ -917,6 +927,24 @@ bool FAnalyticsPropertyStore::Flush(bool bAsync, const FTimespan& Timeout)
 		return false;
 	}
 
+	// If we are inside an AutoRTFM transaction, we can't flush yet, since this potentially involves writing
+	// to disk and spinning up threads. 
+	if (AutoRTFM::IsClosed())
+	{
+		// If we have the option to decline, just return false.
+		if (Timeout < FTimespan::MaxValue())
+		{
+			return false;
+		}
+
+		// The caller expects the flush to succeed; queue up a blocking flush.
+		AutoRTFM::PushOnCommitHandler(this, [this, bAsync]
+		{
+			Flush(bAsync, FTimespan::MaxValue());
+		});
+		return true;
+	}
+
 	TRACE_CPUPROFILER_EVENT_SCOPE(FAnalyticsPropertyStore::Flush);
 
 	// Local helper function to flush.
@@ -925,7 +953,7 @@ bool FAnalyticsPropertyStore::Flush(bool bAsync, const FTimespan& Timeout)
 		// Ensure the previous task completed. Only one thread at a time can write the file.
 		if (FlushTask.IsValid())
 		{
-			// If the task completed or can be successufully canceled.
+			// If the task completed or can be successfully canceled.
 			if (FlushTask->IsDone() || FlushTask->Cancel())
 			{
 				FlushTask.Reset();
@@ -961,36 +989,39 @@ bool FAnalyticsPropertyStore::Flush(bool bAsync, const FTimespan& Timeout)
 	};
 
 	// If the caller doesn't want to wait for a previous execution (if any) to complete.
-	if (Timeout.IsZero() && StoreLock.TryLock())
+	if (Timeout.IsZero())
 	{
-		ON_SCOPE_EXIT { StoreLock.Unlock(); };
+		if (!StoreLock.TryLock())
+		{
+			return false;
+		}
+		ON_SCOPE_EXIT{ StoreLock.Unlock(); };
 		return FlushLockedFn(/*bNoWaiting*/true);
 	}
+
 	// If the caller is ready to wait indefinitely to get its request started.
-	else if (Timeout == FTimespan::MaxValue())
+	if (Timeout == FTimespan::MaxValue())
 	{
-		FScopeLock ScopeLock(&StoreLock);
+		UE::TScopeLock ScopeLock(StoreLock);
 		return FlushLockedFn(/*bNoWaiting*/false);
 	}
+	
 	// The caller is willing to wait some time to get its request processed.
-	else
+	double StartTimeSecs = FPlatformTime::Seconds();
+	do
 	{
-		double StartTimeSecs = FPlatformTime::Seconds();
-		do
+		if (StoreLock.TryLock())
 		{
-			if (StoreLock.TryLock())
+			ON_SCOPE_EXIT { StoreLock.Unlock(); };
+			if (FlushLockedFn(/*bNoWaiting*/true))
 			{
-				ON_SCOPE_EXIT { StoreLock.Unlock(); };
-				if (FlushLockedFn(/*bNoWaiting*/true))
-				{
-					return true;
-				}
+				return true;
 			}
-			FPlatformProcess::Sleep(FloatCastChecked<float>(FTimespan::FromMilliseconds(1).GetTotalSeconds(), 1./1000.));
-		} while (FPlatformTime::Seconds() - StartTimeSecs < Timeout.GetTotalSeconds());
+		}
+		FPlatformProcess::Sleep(FloatCastChecked<float>(FTimespan::FromMilliseconds(1).GetTotalSeconds(), 1./1000.));
+	} while (FPlatformTime::Seconds() - StartTimeSecs < Timeout.GetTotalSeconds());
 
-		return false;
-	}
+	return false;
 }
 
 void FAnalyticsPropertyStore::FlushInternal(const TArray<uint8>& Data)
diff --git a/Engine/Source/Runtime/Analytics/Analytics/Public/AnalyticsPropertyStore.h b/Engine/Source/Runtime/Analytics/Analytics/Public/AnalyticsPropertyStore.h
--- a/Engine/Source/Runtime/Analytics/Analytics/Public/AnalyticsPropertyStore.h
+++ b/Engine/Source/Runtime/Analytics/Analytics/Public/AnalyticsPropertyStore.h
@@ -13,6 +13,7 @@
 #include "Misc/DateTime.h"
 #include "Misc/EnumClassFlags.h"
 #include "Misc/Timespan.h"
+#include "Misc/TransactionallySafeCriticalSection.h"
 #include "Serialization/Archive.h"
 #include "Serialization/MemoryReader.h"
 #include "Serialization/MemoryWriter.h"
@@ -178,7 +179,7 @@ private:
 	}
 
 private:
-	mutable FCriticalSection StoreLock;
+	mutable FTransactionallySafeCriticalSection StoreLock;
 	TArray<uint8> StorageBuf;
 	FMemoryWriter StorageWriter;
 	mutable FMemoryReader StorageReader;
diff --git a/Engine/Source/Runtime/Core/Public/Async/TransactionallySafeRecursiveMutex.h b/Engine/Source/Runtime/Core/Public/Async/TransactionallySafeRecursiveMutex.h
--- a/Engine/Source/Runtime/Core/Public/Async/TransactionallySafeRecursiveMutex.h
+++ b/Engine/Source/Runtime/Core/Public/Async/TransactionallySafeRecursiveMutex.h
@@ -99,6 +99,20 @@ struct FTransactionallySafeRecursiveMutexDefinition final
 		}
 	}
 
+	bool TryLock()
+	{
+		if (AutoRTFM::IsTransactional() || AutoRTFM::IsCommittingOrAborting())
+		{
+			// TODO: implement transactional TryLock. #jira SOL-8303
+			return false;
+		}
+		else
+		{
+			ensure(0 == State->TransactionalLockCount);
+			return State->Mutex.TryLock();
+		}
+	}
+
 private:
 	UE_NONCOPYABLE(FTransactionallySafeRecursiveMutexDefinition)
 
