Update the mobile funnel telemetry to include an event that fires as soon as download is 100% complete.

- Moving the iOS funnel event to a background URLSession

Detailed explanation:
The main goal was to be able to send an analytics event immediately, instead of queueing it to be sent using the threaded system we have nowadays (for iOS).
The current system does not work since the game stops ticking when backgrounded.
So, we took advantage of the iOS background URLSession that supports downloads and uploads.
Making it possible to schedule the upload task with the analytics data as the body (memory file).
This change opens a new pathway for platforms to implement support for making requests that are system handled (or any other custom implementation), while in the background.
So the first use case is for sending analytics events on iOS (the only platform implementing it ATM), but we could research adding more platforms to it when supported.

#rnx

diff --git a/Engine/Source/Runtime/Analytics/Analytics/Private/Interfaces/IAnalyticsProvider.cpp b/Engine/Source/Runtime/Analytics/Analytics/Private/Interfaces/IAnalyticsProvider.cpp
new file mode 100644
--- /dev/null
+++ b/Engine/Source/Runtime/Analytics/Analytics/Private/Interfaces/IAnalyticsProvider.cpp
@@ -0,0 +1,16 @@
+// Copyright Epic Games, Inc. All Rights Reserved.
+
+#include "Interfaces/IAnalyticsProvider.h"
+
+const TCHAR* LexToString(EAnalyticsRecordEventMode Mode)
+{
+    switch (Mode)
+    {
+        case EAnalyticsRecordEventMode::Cached: return TEXT("Cached");
+        case EAnalyticsRecordEventMode::Immediate: return TEXT("Immediate");
+
+        default:
+        checkf(false, TEXT("Please, implement new modes here"));
+        return TEXT("Unknown");
+    }
+}
\ No newline at end of file
diff --git a/Engine/Source/Runtime/Analytics/Analytics/Public/Interfaces/IAnalyticsProvider.h b/Engine/Source/Runtime/Analytics/Analytics/Public/Interfaces/IAnalyticsProvider.h
--- a/Engine/Source/Runtime/Analytics/Analytics/Public/Interfaces/IAnalyticsProvider.h
+++ b/Engine/Source/Runtime/Analytics/Analytics/Public/Interfaces/IAnalyticsProvider.h
@@ -2,11 +2,24 @@
 
 #pragma once
 
-#include "AnalyticsEventAttribute.h"
 #include "CoreMinimal.h"
+#include "AnalyticsEventAttribute.h"
 
 class Error;
 
+/**
+ * Enum describing the supported record event modes
+ * @enum Cached use the existing system where the events are cached and flushed when possible
+ * @enum Immediate use the new system where events are immediately flushed to a HTTP request
+ */
+enum class EAnalyticsRecordEventMode : uint8
+{
+	Cached,
+	Immediate
+};
+
+const TCHAR* LexToString(EAnalyticsRecordEventMode Mode);
+
 /** 
  * Generic interface for an analytics provider. 
  * Other modules can define more and register them with this module. See FAnalytics for details.
@@ -190,6 +203,18 @@ public:
 	 */
 	virtual void RecordEvent(const FString& EventName, const TArray<FAnalyticsEventAttribute>& Attributes) = 0;
 
+	/**
+	 * RecordEvent overload to allow platforms to support events that are immediately sent
+	 *
+	 * @param EventName name of the event
+	 * @param Attributes array of attribute name/value pairs
+	 * @param Mode the record event mode, Cached or Immediate @see EAnalyticsRecordEventMode
+	 */
+	virtual void RecordEvent(const FString& EventName, const TArray<FAnalyticsEventAttribute>& Attributes, EAnalyticsRecordEventMode Mode)
+	{
+		RecordEvent(EventName, Attributes);
+	}
+
 	/**
 	 * Overload for RecordEvent that takes no parameters
 	 *
diff --git a/Engine/Source/Runtime/Analytics/AnalyticsET/Private/AnalyticsProviderETEventCache.cpp b/Engine/Source/Runtime/Analytics/AnalyticsET/Private/AnalyticsProviderETEventCache.cpp
--- a/Engine/Source/Runtime/Analytics/AnalyticsET/Private/AnalyticsProviderETEventCache.cpp
+++ b/Engine/Source/Runtime/Analytics/AnalyticsET/Private/AnalyticsProviderETEventCache.cpp
@@ -346,6 +346,14 @@ TArray<uint8> FAnalyticsProviderETEventCache::FlushCacheUTF8()
 	return TArray<uint8>();
 }
 
+TArray<uint8> FAnalyticsProviderETEventCache::CreateImmediatePayload(const FString& EventName, const TArray<FAnalyticsEventAttribute>& Attributes) const
+{
+	TSharedRef<FAnalyticsProviderETEventCache> Clone = CloneForImmediateMode();
+	Clone->AddToCache(EventName, Attributes);
+	return Clone->FlushCacheUTF8();
+	// the Clone will be released here
+}
+
 // !!!!!!!!!!!!!!!!!!!!!!!!!!!!!!!!!!!!!!!!!!!!!!!!!!!!!!!!!!!!!!!!!!!!!!
 // !!!! This method tries extremely hard to avoid any dynamic allocations
 // !!!! to optimize the flush time. Please don't add new allocations to this function 
@@ -425,6 +433,14 @@ void FAnalyticsProviderETEventCache::QueueFlush()
 	EventCacheStatic::InitializePayloadBuffer(CachedEventUTF8Stream, PreallocatedPayloadSize);
 }
 
+TSharedRef<FAnalyticsProviderETEventCache> FAnalyticsProviderETEventCache::CloneForImmediateMode() const
+{
+	// need to make this because FAnalyticsProviderETEventCache is not allowed to be copyed because
+	// FTransactionallySafeCriticalSection CachedEventsCS depends on FPlatformSharedMutex which has deleted copy and assignment operators.
+	TSharedRef<FAnalyticsProviderETEventCache> Copy = MakeShared<FAnalyticsProviderETEventCache>(MaximumPayloadSize, PreallocatedPayloadSize);
+	Copy->CachedDefaultAttributeUTF8Stream = CachedDefaultAttributeUTF8Stream;
+	return Copy;
+}
 
 bool FAnalyticsProviderETEventCache::CanFlush() const
 {
diff --git a/Engine/Source/Runtime/Analytics/AnalyticsET/Private/IAnalyticsProviderET.cpp b/Engine/Source/Runtime/Analytics/AnalyticsET/Private/IAnalyticsProviderET.cpp
--- a/Engine/Source/Runtime/Analytics/AnalyticsET/Private/IAnalyticsProviderET.cpp
+++ b/Engine/Source/Runtime/Analytics/AnalyticsET/Private/IAnalyticsProviderET.cpp
@@ -143,6 +143,7 @@ public:
 
 	virtual bool ShouldRecordEvent(const FString& EventName) const override;
 	virtual void RecordEvent(FString&& EventName, const TArray<FAnalyticsEventAttribute>& Attributes) override;
+	virtual void RecordEvent(FString&& EventName, const TArray<FAnalyticsEventAttribute>& Attributes, EAnalyticsRecordEventMode Mode) override;
 	virtual void SetDefaultEventAttributes(TArray<FAnalyticsEventAttribute>&& Attributes) override;
 	virtual TArray<FAnalyticsEventAttribute> GetDefaultEventAttributesSafe() const override;
 	virtual int32 GetDefaultEventAttributeCount() const override;
@@ -161,6 +162,8 @@ public:
 	virtual const FAnalyticsET::Config& GetConfig() const override { return Config; }
 
 private:
+	void ExecuteRequest(TArray<uint8>& Payload, OUT int32& PayloadSize, OUT int32& EventCount, EAnalyticsRecordEventMode Mode);
+	void SendImmediately(const FString& EventName, const TArray<FAnalyticsEventAttribute>& Attributes);
 	void FlushEventsOnce();
 	void FlushEventLegacy(const FString& EventName, const TArray<FAnalyticsEventAttribute>& Attributes);
 	bool IsActingAsNullProvider() const
@@ -171,7 +174,7 @@ private:
 	}
 
 	/** Create a request utilizing HttpRetry domains */
-	TSharedRef<IHttpRequest, ESPMode::ThreadSafe> CreateRequest();
+	TSharedRef<IHttpRequest, ESPMode::ThreadSafe> CreateRequest(EAnalyticsRecordEventMode Mode = EAnalyticsRecordEventMode::Cached);
 
 	bool bSessionInProgress;
 	/** The current configuration (might be updated with respect to the one provided at construction). */
@@ -422,7 +425,7 @@ void FAnalyticsProviderET::EndSession()
 	bSessionInProgress = false;
 }
 
-TSharedRef<IHttpRequest, ESPMode::ThreadSafe> FAnalyticsProviderET::CreateRequest()
+TSharedRef<IHttpRequest, ESPMode::ThreadSafe> FAnalyticsProviderET::CreateRequest(EAnalyticsRecordEventMode Mode)
 {
 	LLM_SCOPE_BYNAME(TEXT("Analytics"));
 	if (!ensure(FModuleManager::Get().IsModuleLoaded("HTTP")))
@@ -435,12 +438,19 @@ TSharedRef<IHttpRequest, ESPMode::ThreadSafe> FAnalyticsProviderET::CreateReques
 		FHttpRetrySystem::FRetryTimeoutRelativeSecondsSetting(),
 		FHttpRetrySystem::FRetryResponseCodes(),
 		FHttpRetrySystem::FRetryVerbs(),
-		RetryServers);
+		RetryServers,
+		FHttpRetrySystem::FRetryLimitCountSetting(),
+		FHttpRetrySystem::FExponentialBackoffCurve());
 	for (const TPair<FString, FString>& HttpHeader : HttpHeaders)
 	{
 		HttpRequest->SetHeader(HttpHeader.Key, HttpHeader.Value);
 	}
 
+	if (Mode == EAnalyticsRecordEventMode::Immediate)
+	{
+		HttpRequest->SetOption(HttpRequestOptions::RequestMode, LexToString(EHttpRequestMode::ImmediateRequest));
+	}
+
 	return HttpRequest;
 }
 
@@ -460,6 +470,86 @@ void FAnalyticsProviderET::FlushEvents()
 	};
 }
 
+void FAnalyticsProviderET::ExecuteRequest(TArray<uint8>& Payload, OUT int32& PayloadSize, OUT int32& EventCount, EAnalyticsRecordEventMode Mode)
+{
+	EventCount = 0;
+	PayloadSize = 0;
+
+	// UrlEncode NOTE: need to concatenate everything
+	FString URLPath = Config.APIEndpointET;
+			URLPath += TEXT("?SessionID=") + FPlatformHttp::UrlEncode(SessionID);
+			URLPath += TEXT("&AppID=") + FPlatformHttp::UrlEncode(Config.APIKeyET);
+			URLPath += TEXT("&AppVersion=") + FPlatformHttp::UrlEncode(Config.AppVersionET);
+			URLPath += TEXT("&UserID=") + FPlatformHttp::UrlEncode(UserID);
+			URLPath += TEXT("&AppEnvironment=") + FPlatformHttp::UrlEncode(Config.AppEnvironment);
+			URLPath += TEXT("&UploadType=") + FPlatformHttp::UrlEncode(Config.UploadType);
+	PayloadSize = URLPath.Len() + Payload.Num();
+
+	// Recreate the URLPath for logging because we do not want to escape the parameters when logging.
+	// We cannot simply UrlEncode the entire Path after logging it because UrlEncode(Params) != UrlEncode(Param1) & UrlEncode(Param2) ...
+	UE_LOG(LogAnalytics, VeryVerbose, TEXT("[%s] AnalyticsET URL:%s?SessionID=%s&AppID=%s&AppVersion=%s&UserID=%s&AppEnvironment=%s&UploadType=%s. Payload:%hs"),
+		*Config.APIKeyET,
+		*Config.APIEndpointET,
+		*SessionID,
+		*Config.APIKeyET,
+		*Config.AppVersionET,
+		*UserID,
+		*Config.AppEnvironment,
+		*Config.UploadType,
+		FUtf8StringView((FGenericPlatformTypes::UTF8CHAR*)Payload.GetData(), Payload.Num()).GetData());
+
+	{
+		QUICK_SCOPE_CYCLE_COUNTER(STAT_FlushEventsHttpRequest);
+
+		// Create/send Http request for an event
+		TSharedRef<IHttpRequest, ESPMode::ThreadSafe> HttpRequest = CreateRequest(Mode);
+		HttpRequest->SetHeader(TEXT("Content-Type"), TEXT("application/json; charset=utf-8"));
+		// Want to avoid putting the project name into the User-Agent, because for some apps (like the editor), the project name is private info.
+		// The analytics User-Agent uses the default User-Agent, but with project name removed.
+		HttpRequest->SetHeader(TEXT("User-Agent"), UserAgentCache.GetUserAgent());
+		HttpRequest->SetURL(Config.APIServerET / URLPath);
+		HttpRequest->SetVerb(TEXT("POST"));
+		HttpRequest->SetContent(MoveTemp(Payload));
+
+#if UE_HTTP_SUPPORT_UNIX_SOCKET
+		if (!Config.APIUnixSocketPathET.IsEmpty())
+		{
+			HttpRequest->SetOption(HttpRequestOptions::UnixSocketPath, Config.APIUnixSocketPathET);
+		}
+#endif //UE_HTTP_SUPPORT_UNIX_SOCKET
+
+		// Don't set a response callback if we are in our destructor, as the instance will no longer be there to call.
+		if (!bInDestructor)
+		{
+			HttpRequest->OnProcessRequestComplete().BindSP(this, &FAnalyticsProviderET::EventRequestComplete);
+		}
+
+		OnPreAnalyticsEventProcessedDelegate.Broadcast(HttpRequest);
+
+		HttpRequest->ProcessRequest();
+	}
+}
+
+void FAnalyticsProviderET::SendImmediately(const FString& EventName, const TArray<FAnalyticsEventAttribute>& Attributes)
+{
+	LLM_SCOPE_BYNAME(TEXT("Analytics"));
+
+	if (ensure(FModuleManager::Get().IsModuleLoaded("HTTP")))
+	{
+		TArray<uint8> Payload = EventCache.CreateImmediatePayload(EventName, Attributes);
+
+		int32 PayloadSize = 0;
+		int32 EventCount = 0;
+		ExecuteRequest(Payload, PayloadSize, EventCount, EAnalyticsRecordEventMode::Immediate);
+	}
+	else
+	{
+		// If the HTTP module is not loaded yet for some reason, we fallback to the Cached system so we won't lose this event.
+		// This should not happen, but better safe than sorry.
+		RecordEvent(CopyTemp(EventName), Attributes, EAnalyticsRecordEventMode::Cached);
+	}
+}
+
 void FAnalyticsProviderET::FlushEventsOnce()
 {
 	// FlushEventsOnce cannot be rolled back, so it must not occur inside an AutoRTFM transaction.
@@ -482,66 +572,8 @@ void FAnalyticsProviderET::FlushEventsOnce()
 	if(ensure(FModuleManager::Get().IsModuleLoaded("HTTP")))
 	{
 		TArray<uint8> Payload = EventCache.FlushCacheUTF8();
-		// UrlEncode NOTE: need to concatenate everything
-		FString URLPath = Config.APIEndpointET;
-				URLPath += TEXT("?SessionID=") + FPlatformHttp::UrlEncode(SessionID);
-				URLPath += TEXT("&AppID=") + FPlatformHttp::UrlEncode(Config.APIKeyET);
-				URLPath += TEXT("&AppVersion=") + FPlatformHttp::UrlEncode(Config.AppVersionET);
-				URLPath += TEXT("&UserID=") + FPlatformHttp::UrlEncode(UserID);
-				URLPath += TEXT("&AppEnvironment=") + FPlatformHttp::UrlEncode(Config.AppEnvironment);
-				URLPath += TEXT("&UploadType=") + FPlatformHttp::UrlEncode(Config.UploadType);
-		PayloadSize = URLPath.Len() + Payload.Num();
-
-		// This should never be done in production. MUCH slower!
-		if (UE_LOG_ACTIVE(LogAnalytics, VeryVerbose))
-		{
-			// need to null terminate to load the payload.
-			Payload.Add(TEXT('\0'));
-			// Recreate the URLPath for logging because we do not want to escape the parameters when logging.
-			// We cannot simply UrlEncode the entire Path after logging it because UrlEncode(Params) != UrlEncode(Param1) & UrlEncode(Param2) ...
-			FString LogString = FString::Printf(TEXT("[%s] AnalyticsET URL:%s?SessionID=%s&AppID=%s&AppVersion=%s&UserID=%s&AppEnvironment=%s&UploadType=%s. Payload:%s"),
-				*Config.APIKeyET,
-				*Config.APIEndpointET,
-				*SessionID,
-				*Config.APIKeyET,
-				*Config.AppVersionET,
-				*UserID,
-				*Config.AppEnvironment,
-				*Config.UploadType,
-				UTF8_TO_TCHAR(Payload.GetData()));
-			UE_LOG(LogAnalytics, VeryVerbose, TEXT("%s"), *LogString);
-			Payload.SetNum(Payload.Num()-1);
-		}
-
-		{
-			QUICK_SCOPE_CYCLE_COUNTER(STAT_FlushEventsHttpRequest);
-			// Create/send Http request for an event
-			TSharedRef<IHttpRequest, ESPMode::ThreadSafe> HttpRequest = CreateRequest();
-			HttpRequest->SetHeader(TEXT("Content-Type"), TEXT("application/json; charset=utf-8"));
-			// Want to avoid putting the project name into the User-Agent, because for some apps (like the editor), the project name is private info.
-			// The analytics User-Agent uses the default User-Agent, but with project name removed.
-			HttpRequest->SetHeader(TEXT("User-Agent"), UserAgentCache.GetUserAgent());
-			HttpRequest->SetURL(Config.APIServerET / URLPath);
-			HttpRequest->SetVerb(TEXT("POST"));
-			HttpRequest->SetContent(MoveTemp(Payload));
-
-#if UE_HTTP_SUPPORT_UNIX_SOCKET
-			if (!Config.APIUnixSocketPathET.IsEmpty())
-			{
-				HttpRequest->SetOption(HttpRequestOptions::UnixSocketPath, Config.APIUnixSocketPathET);
-			}
-#endif //UE_HTTP_SUPPORT_UNIX_SOCKET
-
-			// Don't set a response callback if we are in our destructor, as the instance will no longer be there to call.
-			if (!bInDestructor)
-			{
-				HttpRequest->OnProcessRequestComplete().BindSP(this, &FAnalyticsProviderET::EventRequestComplete);
-			}
-
-			OnPreAnalyticsEventProcessedDelegate.Broadcast(HttpRequest);
-
-			HttpRequest->ProcessRequest();
-		}
+		
+		ExecuteRequest(Payload, PayloadSize, EventCount, EAnalyticsRecordEventMode::Cached);
 	}
 	ANALYTICS_FLUSH_TRACKING_END(PayloadSize, EventCount);
 }
@@ -619,29 +651,51 @@ void FAnalyticsProviderET::RecordEvent(FString&& EventName, const TArray<FAnalyt
 {
 	LLM_SCOPE_BYNAME(TEXT("Analytics"));
 
+	RecordEvent(MoveTemp(EventName), Attributes, EAnalyticsRecordEventMode::Cached);
+}
+
+void FAnalyticsProviderET::RecordEvent(FString&& EventName, const TArray<FAnalyticsEventAttribute>& Attributes, EAnalyticsRecordEventMode Mode)
+{
 	// let higher level code filter the decision of whether to send the event
 	if (ShouldRecordEvent(EventName))
 	{
-		// fire any callbacks
-		for (const auto& Cb : EventRecordedCallbacks)
+		switch (Mode)
 		{
-			// we no longer track if the event was Json, each attribute does.
-			Cb(EventName, Attributes, false);
-		}
-
-		if (!Config.UseLegacyProtocol)
-		{
-			EventCache.AddToCache(MoveTemp(EventName), Attributes);
+			case EAnalyticsRecordEventMode::Cached:
+			{
+				// fire any callbacks
+				for (const auto& Cb : EventRecordedCallbacks)
+				{
+					// we no longer track if the event was Json, each attribute does.
+					Cb(EventName, Attributes, false);
+				}
+
+				if (!Config.UseLegacyProtocol)
+				{
+					EventCache.AddToCache(MoveTemp(EventName), Attributes);
+
+					// if we aren't caching events, flush immediately. This is really only for debugging as it will significantly affect bandwidth.
+					if (!bShouldCacheEvents)
+					{
+						FlushEvents();
+					}
+				}
+				else
+				{
+					FlushEventLegacy(EventName, Attributes);
+				}
+			}
+			break;
 
-			// if we aren't caching events, flush immediately. This is really only for debugging as it will significantly affect bandwidth.
-			if (!bShouldCacheEvents)
+			case EAnalyticsRecordEventMode::Immediate:
 			{
-				FlushEvents();
+				SendImmediately(MoveTemp(EventName), Attributes);
 			}
-		}
-		else
-		{
-			FlushEventLegacy(EventName, Attributes);
+			break;
+
+			default:
+			checkf(false, TEXT("Please, implement new modes here"));
+			break;
 		}
 	}
 	else
diff --git a/Engine/Source/Runtime/Analytics/AnalyticsET/Public/AnalyticsProviderETEventCache.h b/Engine/Source/Runtime/Analytics/AnalyticsET/Public/AnalyticsProviderETEventCache.h
--- a/Engine/Source/Runtime/Analytics/AnalyticsET/Public/AnalyticsProviderETEventCache.h
+++ b/Engine/Source/Runtime/Analytics/AnalyticsET/Public/AnalyticsProviderETEventCache.h
@@ -71,6 +71,11 @@ public:
 	/** Flushes the cache as a UTF8 char array. Returns a uint8 because that's what IHttpRequest prefers. If nothing is cached, returns an empty array. */
 	ANALYTICSET_API TArray<uint8> FlushCacheUTF8();
 
+	/**
+	 * Creates a immediate payload with the existing default attributes
+	 */
+	ANALYTICSET_API TArray<uint8> CreateImmediatePayload(const FString& EventName, const TArray<FAnalyticsEventAttribute>& Attributes) const;
+
 	/**
 	* Determines whether we have anything we need to flush, either a queued flush or existing events in the payload.
 	*/
@@ -115,6 +120,8 @@ public:
 private:
 	ANALYTICSET_API void QueueFlush();
 
+	ANALYTICSET_API TSharedRef<FAnalyticsProviderETEventCache> CloneForImmediateMode() const;
+
 	/**
 	* Analytics event entry to be cached
 	*/
diff --git a/Engine/Source/Runtime/Analytics/AnalyticsET/Public/IAnalyticsProviderET.h b/Engine/Source/Runtime/Analytics/AnalyticsET/Public/IAnalyticsProviderET.h
--- a/Engine/Source/Runtime/Analytics/AnalyticsET/Public/IAnalyticsProviderET.h
+++ b/Engine/Source/Runtime/Analytics/AnalyticsET/Public/IAnalyticsProviderET.h
@@ -37,6 +37,14 @@ public:
 		return RecordEvent(CopyTemp(EventName), Attributes);
 	}
 
+	/**
+	 * This class augments RecordEvent with a version that takes the EventName by rvalue reference to save a string copy. Implement the base version in terms of this one.
+	 */
+	virtual void RecordEvent(const FString& EventName, const TArray<FAnalyticsEventAttribute>& Attributes, EAnalyticsRecordEventMode Mode) override
+	{
+		return RecordEvent(CopyTemp(EventName), Attributes, Mode);
+	}
+
 	/**
 	 * This class augments StartSession with a version that takes the SessionID instead of always generating it. Implement the base version in terms of this one.
 	 */
@@ -96,6 +104,11 @@ public:
 	 * Primary RecordEvent API. Allow move semantics to capture the attributes.
 	 */
 	virtual void RecordEvent(FString&& EventName, const TArray<FAnalyticsEventAttribute>& Attributes) = 0;
+	
+	/**
+	 * Overload to allow platforms to implement events that are immediately sent when supported. @see EAnalyticsRecordEventMode.
+	 */
+	virtual void RecordEvent(FString&& EventName, const TArray<FAnalyticsEventAttribute>& Attributes, EAnalyticsRecordEventMode Mode) = 0;
 
 	/**
 	 * Updates the default Domain and AltDomains.
diff --git a/Engine/Source/Runtime/ApplicationCore/Private/IOS/IOSAppDelegate.cpp b/Engine/Source/Runtime/ApplicationCore/Private/IOS/IOSAppDelegate.cpp
--- a/Engine/Source/Runtime/ApplicationCore/Private/IOS/IOSAppDelegate.cpp
+++ b/Engine/Source/Runtime/ApplicationCore/Private/IOS/IOSAppDelegate.cpp
@@ -263,7 +263,6 @@ bool FIOSCoreDelegates::PassesPushNotificationFilters(NSDictionary* Payload)
 @synthesize AccessibilityCacheTimer;
 #endif
 @synthesize savedOpenUrlParameters;
-@synthesize BackgroundSessionEventCompleteDelegate;
 
 
 static IOSAppDelegate* CachedDelegate = nil;
@@ -1872,7 +1871,7 @@ extern double GCStartTime;
 - (void)application:(UIApplication *)application handleEventsForBackgroundURLSession:(NSString *)identifier completionHandler:(void(^)(void))completionHandler
 {
     //Save off completionHandler so that a future call to FCoreDelegates::ApplicationBackgroundSessionEventsAllSentDelegate can execute it
-    self.BackgroundSessionEventCompleteDelegate = completionHandler;
+	[FBackgroundURLSessionHandler::BackgroundSessionEventCompleteDelegateMap setObject:completionHandler forKey:identifier];
     
     // Invoke background session with this identifier if needed to handle these events
 	FString Id(identifier);
diff --git a/Engine/Source/Runtime/ApplicationCore/Private/IOS/IOSBackgroundURLSessionHandler.cpp b/Engine/Source/Runtime/ApplicationCore/Private/IOS/IOSBackgroundURLSessionHandler.cpp
--- a/Engine/Source/Runtime/ApplicationCore/Private/IOS/IOSBackgroundURLSessionHandler.cpp
+++ b/Engine/Source/Runtime/ApplicationCore/Private/IOS/IOSBackgroundURLSessionHandler.cpp
@@ -1427,14 +1427,15 @@ static constexpr NSInteger HTTPStatusCodeErrorServer = 500;
 {
 	UE_DNLD_LOG(@"URLSessionDidFinishEventsForBackgroundURLSession");
 
-	IOSAppDelegate* AppDelegate = [IOSAppDelegate GetDelegate];
-	if (AppDelegate == nullptr || AppDelegate.BackgroundSessionEventCompleteDelegate == nil)
+	NSString* Id = Session.configuration.identifier;
+
+	if ([FBackgroundURLSessionHandler::BackgroundSessionEventCompleteDelegateMap objectForKey:Id] == nil)
 	{
 		return;
 	}
 	
-	void(^CompletionHandler)() = [AppDelegate.BackgroundSessionEventCompleteDelegate retain];
-	AppDelegate.BackgroundSessionEventCompleteDelegate = nil;
+	void(^CompletionHandler)() = [[FBackgroundURLSessionHandler::BackgroundSessionEventCompleteDelegateMap objectForKey:Id] retain];
+	[FBackgroundURLSessionHandler::BackgroundSessionEventCompleteDelegateMap removeObjectForKey:Id];
 
 	// Completion handler has to be invoked on the main thread.
 	[[NSOperationQueue mainQueue] addOperationWithBlock:^
@@ -1647,6 +1648,8 @@ FBackgroundURLSessionHandler::FOnDownloadMetrics FBackgroundURLSessionHandler::O
 
 FBackgroundURLSessionHandler::FOnDownloadsCompletedWhileAppWasNotRunning FBackgroundURLSessionHandler::OnDownloadsCompletedWhileAppWasNotRunning;
 
+NSMutableDictionary<NSString*,void(^)()>* FBackgroundURLSessionHandler::BackgroundSessionEventCompleteDelegateMap = [[NSMutableDictionary alloc] init];
+
 void FBackgroundURLSessionHandler::AllowCellular(bool bAllow)
 {
 	@autoreleasepool
diff --git a/Engine/Source/Runtime/ApplicationCore/Public/IOS/IOSAppDelegate.h b/Engine/Source/Runtime/ApplicationCore/Public/IOS/IOSAppDelegate.h
--- a/Engine/Source/Runtime/ApplicationCore/Public/IOS/IOSAppDelegate.h
+++ b/Engine/Source/Runtime/ApplicationCore/Public/IOS/IOSAppDelegate.h
@@ -302,8 +302,6 @@ APPLICATIONCORE_API
 
 - (void)ForceExit;
 
-@property (nonatomic, copy) void(^BackgroundSessionEventCompleteDelegate)();
-
 static void interruptionListener(void* ClientData, UInt32 Interruption);
 
 -(UIWindow*)window;
diff --git a/Engine/Source/Runtime/ApplicationCore/Public/IOS/IOSBackgroundURLSessionHandler.h b/Engine/Source/Runtime/ApplicationCore/Public/IOS/IOSBackgroundURLSessionHandler.h
--- a/Engine/Source/Runtime/ApplicationCore/Public/IOS/IOSBackgroundURLSessionHandler.h
+++ b/Engine/Source/Runtime/ApplicationCore/Public/IOS/IOSBackgroundURLSessionHandler.h
@@ -70,6 +70,9 @@ public:
 	DECLARE_MULTICAST_DELEGATE_OneParam(FOnDownloadsCompletedWhileAppWasNotRunning, const bool /*bSuccess*/);
 	static FOnDownloadsCompletedWhileAppWasNotRunning OnDownloadsCompletedWhileAppWasNotRunning;
 
+	// Will be used by IOSAppDelegate and ApplePlatformHttp
+	static NSMutableDictionary<NSString*,void(^)()>* BackgroundSessionEventCompleteDelegateMap;
+
 	// Sets if cellular is allowed to be used for new downloads.
 	// Existing downloads will be recreated to reflect new setting value.
 	static void AllowCellular(bool bAllow);
diff --git a/Engine/Source/Runtime/Online/HTTP/HTTP.Build.cs b/Engine/Source/Runtime/Online/HTTP/HTTP.Build.cs
--- a/Engine/Source/Runtime/Online/HTTP/HTTP.Build.cs
+++ b/Engine/Source/Runtime/Online/HTTP/HTTP.Build.cs
@@ -86,6 +86,16 @@ public class HTTP : ModuleRules
 			}
 			);
 
+		if (Target.bCompileAgainstApplicationCore && Target.Platform == UnrealTargetPlatform.IOS)
+		{
+			PrivateDependencyModuleNames.AddRange(
+				new string[] {
+					"ApplicationCore",
+					"BackgroundHTTPFileHash",
+				}
+			);
+		}
+
 		PrivateDependencyModuleNames.AddRange(
 			new string[] {
 				"EventLoop",
diff --git a/Engine/Source/Runtime/Online/HTTP/Private/Apple/AppleHttp.cpp b/Engine/Source/Runtime/Online/HTTP/Private/Apple/AppleHttp.cpp
--- a/Engine/Source/Runtime/Online/HTTP/Private/Apple/AppleHttp.cpp
+++ b/Engine/Source/Runtime/Online/HTTP/Private/Apple/AppleHttp.cpp
@@ -2,6 +2,8 @@
 
 
 #include "AppleHttp.h"
+#include "Apple/ApplePlatformHttp.h"
+#include "AppleHttpObjc.h"
 #include "HAL/FileManager.h"
 #include "HAL/IConsoleManager.h"
 #include "HAL/PlatformFileManager.h"
@@ -9,62 +11,11 @@
 #include "Http.h"
 #include "HttpManager.h"
 #include "HttpModule.h"
+#include "Interfaces/IHttpRequest.h"
 #include "Misc/EngineVersion.h"
 #include "Misc/App.h"
 #include "Misc/Base64.h"
 
-// UE-228716
-TAutoConsoleVariable<bool> CVarAppleHttpStartTaskInHttpThreadEnabled(
-	TEXT("AppleHttp.StartTaskInHttpThreadEnabled"),
-	true,
-	TEXT("Enables this implementation to start the task in the http thread like the rest of implementations"),
-	ECVF_SaveForNextBoot
-);
-
-/**
- * Class to hold data from delegate implementation notifications.
- */
-
-@interface FAppleHttpResponseDelegate : NSObject<NSURLSessionDataDelegate>
-{
-	// Flag to indicate the request was initialized with stream. In that case even if stream was set to 
-	// null later on internally, the request itself won't cache received data anymore
-	@public BOOL bInitializedWithValidStream;
-
-	/** Have we received any data? */
-	BOOL bAnyHttpActivity;
-
-	/** Delegate invoked after processing URLSession:dataTask:didReceiveData or URLSession:task:didCompleteWithError:*/
-	@public FNewAppleHttpEventDelegate NewAppleHttpEventDelegate;
-}
-
-/** A handle for the response */
-@property(retain) NSURLResponse* Response;
-/** The total number of bytes written out during the request/response */
-@property uint64 BytesWritten;
-/** The total number of bytes received out during the request/response */
-@property uint64 BytesReceived;
-/** Request status */
-@property EHttpRequestStatus::Type RequestStatus;
-/** Reason of failure */
-@property EHttpFailureReason FailureReason;
-/** Associated request. Cleared when canceled */
-@property TWeakPtr<FAppleHttpRequest> SourceRequest;
-
-/** NSURLSessionDataDelegate delegate methods. Those are called from a thread controlled by the NSURLSession */
-
-/** Sent periodically to notify the delegate of upload progress. */
-- (void)URLSession:(NSURLSession *)session task:(NSURLSessionTask *)task didSendBodyData:(int64_t)bytesSent totalBytesSent:(int64_t)totalBytesSent totalBytesExpectedToSend:(int64_t)totalBytesExpectedToSend;
-/** The task has received a response and no further messages will be received until the completion block is called. */
-- (void)URLSession:(NSURLSession *)session dataTask:(NSURLSessionDataTask *)dataTask didReceiveResponse:(NSURLResponse *)response completionHandler:(void (^)(NSURLSessionResponseDisposition disposition))completionHandler;
-/** Sent when data is available for the delegate to consume. Data may be discontiguous */
-- (void)URLSession:(NSURLSession *)session dataTask:(NSURLSessionDataTask *)dataTask didReceiveData:(NSData *)data;
-/** Sent as the last message related to a specific task.  A nil Error implies that no error occurred and this task is complete. */
-- (void)URLSession:(NSURLSession *)session task:(NSURLSessionTask *)task didCompleteWithError:(nullable NSError *)error;
-/** Asks the delegate if it needs to store responses in the cache. */
-- (void)URLSession:(NSURLSession *)session dataTask:(NSURLSessionDataTask *)dataTask willCacheResponse:(NSCachedURLResponse *)proposedResponse completionHandler:(void (^)(NSCachedURLResponse *cachedResponse))completionHandler;
-@end
-
 @implementation FAppleHttpResponseDelegate
 @synthesize Response;
 @synthesize RequestStatus;
@@ -482,13 +433,13 @@ TAutoConsoleVariable<bool> CVarAppleHttpStartTaskInHttpThreadEnabled(
  * FAppleHttpRequest implementation
  ***************************************************************************/
 
-FAppleHttpRequest::FAppleHttpRequest(NSURLSession* InSession)
+FAppleHttpRequest::FAppleHttpRequest(NSURLSession* InSession, NSURLSession* InBackgroundSession)
 :   Session([InSession retain])
+,	BackgroundSession([InBackgroundSession retain])
 ,   Task(nil)
 ,	ContentBytesLength(0)
 ,	LastReportedBytesWritten(0)
 ,	LastReportedBytesRead(0)
-,	bStartTaskInHttpThread(CVarAppleHttpStartTaskInHttpThreadEnabled.GetValueOnAnyThread())
 {
 	bUsePlatformActivityTimeout = false;
 	
@@ -712,12 +663,30 @@ bool FAppleHttpRequest::ProcessRequest()
 	
 	StartProcessTime = FPlatformTime::Seconds();
 
-	if (bStartTaskInHttpThread)
+	SetStatus(EHttpRequestStatus::Processing);
+	SetFailureReason(EHttpFailureReason::None);
+	// AppleEventLoop sets a delegate into the response to be able to notify events
+	InitResponse();
+
+	// if we specified the option to use ImmediateRequest and the feature is enabled (which means the BackgroundSession should not be nil)
+	if (GetOption(HttpRequestOptions::RequestMode) == LexToString(EHttpRequestMode::ImmediateRequest) 
+		&& FApplePlatformHttp::IsBackgroundRequestFeatureEnabled() && BackgroundSession != nil)
+	{
+
+		FHttpModule::Get().GetHttpManager().AddRequest(SharedThis(this));
+
+		// for background sessions we must use the download/upload APIs instead of the dataTask ones
+		// it is a requirement for the background session to handle these requests while in background
+		Task = [BackgroundSession uploadTaskWithStreamedRequest:Request];
+
+		UE_LOG(LogHttp, VeryVerbose, TEXT("Starting immediate request %p with task %p"), this, Task);
+
+		FApplePlatformHttp::Add([Task retain], SharedThis(this));
+
+		[Task resume];
+	}
+	else
 	{
-		SetStatus(EHttpRequestStatus::Processing);
-		SetFailureReason(EHttpFailureReason::None);
-		// AppleEventLoop sets a delegate into the response to be able to notify events
-		InitResponse();
 		FHttpModule::Get().GetHttpManager().AddThreadedRequest(SharedThis(this));
 	}
 	return true;
@@ -792,40 +761,7 @@ bool FAppleHttpRequest::SetupRequest()
 		return false;
 	}
 	
-	if (bStartTaskInHttpThread)
-	{
-		return true;
-	}
-	else
-	{
-		Task = [Session dataTaskWithRequest: Request];
-	
-		if (Task != nil)
-		{
-			SetStatus(EHttpRequestStatus::Processing);
-			SetFailureReason(EHttpFailureReason::None);
-
-			InitResponse();
-
-			// Both Task and Response keep a strong reference to the delegate
-			TSharedPtr<FAppleHttpResponse> Response = StaticCastSharedPtr<FAppleHttpResponse>(ResponseCommon);
-			Task.delegate = Response->ResponseDelegate;
-
-			//Setup delegates before starting the request
-			FHttpModule::Get().GetHttpManager().AddThreadedRequest(SharedThis(this));
-
-			[[Task retain] resume];
-			UE_LOG(LogHttp, Verbose, TEXT("Task started %p"), this);
-			return true;
-		}
-		else
-		{
-			UE_LOG(LogHttp, Warning, TEXT("SetupRequest failed. Could not initialize NSURLSessionTask."));
-			SetStatus(EHttpRequestStatus::Failed);
-			SetFailureReason(EHttpFailureReason::ConnectionError);
-			return false;
-		}	
-	}
+	return true;
 }
 
 FHttpResponsePtr FAppleHttpRequest::CreateResponse()
@@ -979,36 +915,34 @@ void FAppleHttpRequest::CheckProgressDelegate()
 
 bool FAppleHttpRequest::StartThreadedRequest()
 {
-	if (bStartTaskInHttpThread)
+	if (bCanceled)
 	{
-		if (bCanceled)
-		{
-			UE_LOG(LogHttp, Verbose, TEXT("StartThreadedRequest ignored because request has been canceled. %s url=%s"), *GetVerb(), *GetURL());
-			return false;
-		}
+		UE_LOG(LogHttp, Verbose, TEXT("StartThreadedRequest ignored because request has been canceled. %s url=%s"), *GetVerb(), *GetURL());
+		return false;
+	}
 
-		if (Task != nil)
-		{
-			UE_LOG(LogHttp, Verbose, TEXT("StartThreadedRequest ignored because task was already in progress. %s url=%s"), *GetVerb(), *GetURL());
-			return false;
-		}
-	
-		Task = [Session dataTaskWithRequest: Request];
-	
-		if (Task == nil)
-		{
-			UE_LOG(LogHttp, Warning, TEXT("StartThreadedRequest failed. Could not initialize NSURLSessionTask."));
-			SetStatus(EHttpRequestStatus::Failed);
-			SetFailureReason(EHttpFailureReason::ConnectionError);
-			return false;
-		}
+	if (Task != nil)
+	{
+		UE_LOG(LogHttp, Verbose, TEXT("StartThreadedRequest ignored because task was already in progress. %s url=%s"), *GetVerb(), *GetURL());
+		return false;
+	}
 
-		// Both Task and Response keep a strong reference to the delegate
-		TSharedPtr<FAppleHttpResponse> Response = StaticCastSharedPtr<FAppleHttpResponse>(ResponseCommon);
-		Task.delegate = Response->ResponseDelegate;
+	Task = [Session dataTaskWithRequest: Request];
 
-		[[Task retain] resume];
+	if (Task == nil)
+	{
+		UE_LOG(LogHttp, Warning, TEXT("StartThreadedRequest failed. Could not initialize NSURLSessionTask."));
+		SetStatus(EHttpRequestStatus::Failed);
+		SetFailureReason(EHttpFailureReason::ConnectionError);
+		return false;
 	}
+
+	// Both Task and Response keep a strong reference to the delegate
+	TSharedPtr<FAppleHttpResponse> Response = StaticCastSharedPtr<FAppleHttpResponse>(ResponseCommon);
+	Task.delegate = Response->ResponseDelegate;
+
+	[[Task retain] resume];
+	
 	return true;
 }
 
@@ -1055,6 +989,11 @@ void FAppleHttpResponse::SetHeaders(TMap<FString, FString>&& InHeaders)
 	Headers = MoveTemp(InHeaders);
 }
 
+FAppleHttpResponseDelegate* FAppleHttpResponse::GetResponseDelegate() const
+{
+	return ResponseDelegate;
+}
+
 void FAppleHttpResponse::CleanSharedObjects()
 {
 	[ResponseDelegate CleanSharedObjects];
diff --git a/Engine/Source/Runtime/Online/HTTP/Private/Apple/AppleHttp.h b/Engine/Source/Runtime/Online/HTTP/Private/Apple/AppleHttp.h
--- a/Engine/Source/Runtime/Online/HTTP/Private/Apple/AppleHttp.h
+++ b/Engine/Source/Runtime/Online/HTTP/Private/Apple/AppleHttp.h
@@ -57,8 +57,9 @@ public:
 	 * Constructor
 	 *
 	 * @param InSession - NSURLSession session used to create NSURLSessionTask to retrieve the response
+	 * @param InBackgroundSession - NSURLSession session used to create NSURLSessionTask to retrieve the response in the background
 	 */
-	explicit FAppleHttpRequest(NSURLSession* InSession);
+	FAppleHttpRequest(NSURLSession* InSession, NSURLSession* InBackgroundSession);
 
 	/**
 	 * Destructor. Clean up any connection/request handles
@@ -108,9 +109,11 @@ private:
 	/** This is the NSMutableURLRequest, all our Apple functionality will deal with this. */
 	NSMutableURLRequest* Request;
 
-	/** This is the session our request belongs to */
+	/** This is the session our request belongs to (depending on the RequestMode option) */
 	NSURLSession* Session;
-	
+	/** This is the background session our request belongs to (depending on the RequestMode option) */
+	NSURLSession* BackgroundSession;
+
 	/** This is the Task associated to the sessionin charge of our request */
 	NSURLSessionTask* Task;
     
@@ -135,9 +138,6 @@ private:
 
 	/** Last reported bytesread */
 	int32 LastReportedBytesRead;
-	
-	// UE-228716: This request should follow same logic as the rest of implementations
-	const bool bStartTaskInHttpThread;
 };
 
 @class FAppleHttpResponseDelegate;
@@ -182,6 +182,8 @@ public:
 
 	void SetHeaders(TMap<FString, FString>&& InHeaders);
 
+	FAppleHttpResponseDelegate* GetResponseDelegate() const;
+
 	using FHttpResponseCommon::AppendToPayload;
 
 	/**
diff --git a/Engine/Source/Runtime/Online/HTTP/Private/Apple/AppleHttpBackgroundSession.cpp b/Engine/Source/Runtime/Online/HTTP/Private/Apple/AppleHttpBackgroundSession.cpp
new file mode 100644
--- /dev/null
+++ b/Engine/Source/Runtime/Online/HTTP/Private/Apple/AppleHttpBackgroundSession.cpp
@@ -0,0 +1,101 @@
+// Copyright Epic Games, Inc. All Rights Reserved.
+
+#include "AppleHttpBackgroundSession.h"
+#include "Interfaces/IHttpRequest.h"
+
+#if PLATFORM_IOS
+
+#include "Interfaces/IHttpBase.h"
+#include "Http.h"
+#include "AppleHttp.h"
+#include "AppleHttpObjc.h"
+#include "HttpManager.h"
+
+#include "IOS/IOSBackgroundURLSessionHandler.h"
+
+@implementation FApplePlatformHttpBackgroundSessionDelegate
+
+- (void)URLSession:(NSURLSession *)session task:(NSURLSessionTask *)task didSendBodyData:(int64_t)bytesSent totalBytesSent:(int64_t)totalBytesSent totalBytesExpectedToSend:(int64_t)totalBytesExpectedToSend 
+{
+	TSharedPtr<FAppleHttpRequest> Request;
+	if (FApplePlatformHttp::Find(task, Request) && Request.IsValid())
+	{
+		TSharedPtr<FAppleHttpResponse> Response = StaticCastSharedPtr<FAppleHttpResponse>(Request->GetResponse());
+
+		[Response->GetResponseDelegate() URLSession:session task:task didSendBodyData:bytesSent totalBytesSent:totalBytesSent totalBytesExpectedToSend:totalBytesExpectedToSend];
+	}
+}
+
+- (void)URLSession:(NSURLSession *)session dataTask:(NSURLSessionDataTask *)dataTask didReceiveResponse:(NSURLResponse *)response completionHandler:(void (^)(NSURLSessionResponseDisposition))completionHandler 
+{
+	TSharedPtr<FAppleHttpRequest> Request;
+	if (FApplePlatformHttp::Find((NSURLSessionTask*)dataTask, Request) && Request.IsValid())
+	{
+		TSharedPtr<FAppleHttpResponse> Response = StaticCastSharedPtr<FAppleHttpResponse>(Request->GetResponse());
+
+		[Response->GetResponseDelegate() URLSession:session dataTask:dataTask didReceiveResponse:response completionHandler:completionHandler];
+
+		return;
+	}
+	
+	completionHandler(NSURLSessionResponseAllow);
+}
+
+- (void)URLSession:(NSURLSession *)session dataTask:(NSURLSessionDataTask *)dataTask didReceiveData:(NSData *)data 
+{
+	TSharedPtr<FAppleHttpRequest> Request;
+	if (FApplePlatformHttp::Find((NSURLSessionTask*)dataTask, Request) && Request.IsValid())
+	{
+		TSharedPtr<FAppleHttpResponse> Response = StaticCastSharedPtr<FAppleHttpResponse>(Request->GetResponse());
+
+		[Response->GetResponseDelegate() URLSession:session dataTask:dataTask didReceiveData:data];
+	}
+}
+
+- (void)URLSession:(NSURLSession *)session task:(NSURLSessionTask *)task didCompleteWithError:(NSError *)error 
+{
+	// here we gonna remove the task and request from the map since they will be released
+	TSharedPtr<FAppleHttpRequest> Request;
+	if (FApplePlatformHttp::RemoveAndCopyValue(task, Request) && Request.IsValid())
+	{
+		TSharedPtr<FAppleHttpResponse> Response = StaticCastSharedPtr<FAppleHttpResponse>(Request->GetResponse());
+
+		[Response->GetResponseDelegate() URLSession:session task:task didCompleteWithError:error];
+
+		FHttpManager& HttpManager = FHttpModule::Get().GetHttpManager();
+
+		HttpManager.RemoveRequest(Request->AsShared());
+		HttpManager.BroadcastHttpRequestCompleted(Request->AsShared());
+	}
+}
+
+- (void)URLSession:(NSURLSession *)session dataTask:(NSURLSessionDataTask *)dataTask willCacheResponse:(NSCachedURLResponse *)proposedResponse completionHandler:(void (^)(NSCachedURLResponse * _Nullable))completionHandler 
+{
+	// All FAppleHttpRequest use NSURLRequestReloadIgnoringLocalCacheData
+	// NSURLRequestReloadIgnoringLocalCacheData disables loading of data from cache, but responses can still be stored in cache
+	// Passing nil to this handler disables caching the responses
+	completionHandler(nil);
+}
+
+- (void)URLSessionDidFinishEventsForBackgroundURLSession:(NSURLSession *)session 
+{
+	NSString* Id = session.configuration.identifier;
+
+	if ([FBackgroundURLSessionHandler::BackgroundSessionEventCompleteDelegateMap objectForKey:Id] == nil)
+	{
+		return;
+	}
+	
+	void(^CompletionHandler)() = [[FBackgroundURLSessionHandler::BackgroundSessionEventCompleteDelegateMap objectForKey:session.configuration.identifier] retain];
+	[FBackgroundURLSessionHandler::BackgroundSessionEventCompleteDelegateMap removeObjectForKey:Id];
+
+	[[NSOperationQueue mainQueue] addOperationWithBlock:^
+	{
+		CompletionHandler();
+		[CompletionHandler release];
+	}];
+}
+
+@end
+
+#endif // PLATFORM_IOS
\ No newline at end of file
diff --git a/Engine/Source/Runtime/Online/HTTP/Private/Apple/AppleHttpBackgroundSession.h b/Engine/Source/Runtime/Online/HTTP/Private/Apple/AppleHttpBackgroundSession.h
new file mode 100644
--- /dev/null
+++ b/Engine/Source/Runtime/Online/HTTP/Private/Apple/AppleHttpBackgroundSession.h
@@ -0,0 +1,20 @@
+// Copyright Epic Games, Inc. All Rights Reserved.
+
+#pragma once
+
+#if PLATFORM_IOS
+
+#import <Foundation/Foundation.h>
+
+@interface FApplePlatformHttpBackgroundSessionDelegate : NSObject<NSURLSessionDataDelegate>
+
+- (void)URLSession:(NSURLSession *)session task:(NSURLSessionTask *)task didSendBodyData:(int64_t)bytesSent totalBytesSent:(int64_t)totalBytesSent totalBytesExpectedToSend:(int64_t)totalBytesExpectedToSend;
+- (void)URLSession:(NSURLSession *)session dataTask:(NSURLSessionDataTask *)dataTask didReceiveResponse:(NSURLResponse *)response completionHandler:(void (^)(NSURLSessionResponseDisposition))completionHandler;
+- (void)URLSession:(NSURLSession *)session dataTask:(NSURLSessionDataTask *)dataTask didReceiveData:(NSData *)data;
+- (void)URLSession:(NSURLSession *)session task:(NSURLSessionTask *)task didCompleteWithError:(NSError *)error;
+- (void)URLSession:(NSURLSession *)session dataTask:(NSURLSessionDataTask *)dataTask willCacheResponse:(NSCachedURLResponse *)proposedResponse completionHandler:(void (^)(NSCachedURLResponse * _Nullable))completionHandler;
+- (void)URLSessionDidFinishEventsForBackgroundURLSession:(NSURLSession *)session;
+
+@end
+
+#endif // PLATFORM_IOS
\ No newline at end of file
diff --git a/Engine/Source/Runtime/Online/HTTP/Private/Apple/AppleHttpObjc.h b/Engine/Source/Runtime/Online/HTTP/Private/Apple/AppleHttpObjc.h
new file mode 100644
--- /dev/null
+++ b/Engine/Source/Runtime/Online/HTTP/Private/Apple/AppleHttpObjc.h
@@ -0,0 +1,52 @@
+// Copyright Epic Games, Inc. All Rights Reserved.
+
+#pragma once
+
+#include "Interfaces/IHttpBase.h"
+#include "AppleHttp.h"
+
+#import <Foundation/Foundation.h>
+
+/**
+ * Class to hold data from delegate implementation notifications.
+ */
+
+@interface FAppleHttpResponseDelegate : NSObject<NSURLSessionDataDelegate>
+{
+	// Flag to indicate the request was initialized with stream. In that case even if stream was set to 
+	// null later on internally, the request itself won't cache received data anymore
+	@public BOOL bInitializedWithValidStream;
+
+	/** Have we received any data? */
+	BOOL bAnyHttpActivity;
+
+	/** Delegate invoked after processing URLSession:dataTask:didReceiveData or URLSession:task:didCompleteWithError:*/
+	@public FNewAppleHttpEventDelegate NewAppleHttpEventDelegate;
+}
+
+/** A handle for the response */
+@property(retain) NSURLResponse* Response;
+/** The total number of bytes written out during the request/response */
+@property uint64 BytesWritten;
+/** The total number of bytes received out during the request/response */
+@property uint64 BytesReceived;
+/** Request status */
+@property EHttpRequestStatus::Type RequestStatus;
+/** Reason of failure */
+@property EHttpFailureReason FailureReason;
+/** Associated request. Cleared when canceled */
+@property TWeakPtr<FAppleHttpRequest> SourceRequest;
+
+/** NSURLSessionDataDelegate delegate methods. Those are called from a thread controlled by the NSURLSession */
+
+/** Sent periodically to notify the delegate of upload progress. */
+- (void)URLSession:(NSURLSession *)session task:(NSURLSessionTask *)task didSendBodyData:(int64_t)bytesSent totalBytesSent:(int64_t)totalBytesSent totalBytesExpectedToSend:(int64_t)totalBytesExpectedToSend;
+/** The task has received a response and no further messages will be received until the completion block is called. */
+- (void)URLSession:(NSURLSession *)session dataTask:(NSURLSessionDataTask *)dataTask didReceiveResponse:(NSURLResponse *)response completionHandler:(void (^)(NSURLSessionResponseDisposition disposition))completionHandler;
+/** Sent when data is available for the delegate to consume. Data may be discontiguous */
+- (void)URLSession:(NSURLSession *)session dataTask:(NSURLSessionDataTask *)dataTask didReceiveData:(NSData *)data;
+/** Sent as the last message related to a specific task.  A nil Error implies that no error occurred and this task is complete. */
+- (void)URLSession:(NSURLSession *)session task:(NSURLSessionTask *)task didCompleteWithError:(nullable NSError *)error;
+/** Asks the delegate if it needs to store responses in the cache. */
+- (void)URLSession:(NSURLSession *)session dataTask:(NSURLSessionDataTask *)dataTask willCacheResponse:(NSCachedURLResponse *)proposedResponse completionHandler:(void (^)(NSCachedURLResponse *cachedResponse))completionHandler;
+@end
\ No newline at end of file
diff --git a/Engine/Source/Runtime/Online/HTTP/Private/Apple/ApplePlatformHttp.cpp b/Engine/Source/Runtime/Online/HTTP/Private/Apple/ApplePlatformHttp.cpp
--- a/Engine/Source/Runtime/Online/HTTP/Private/Apple/ApplePlatformHttp.cpp
+++ b/Engine/Source/Runtime/Online/HTTP/Private/Apple/ApplePlatformHttp.cpp
@@ -1,12 +1,15 @@
 // Copyright Epic Games, Inc. All Rights Reserved.
 
 #include "Apple/ApplePlatformHttp.h"
+#include "HAL/IConsoleManager.h"
 #include "Http.h"
 #include "AppleHttpManager.h"
+#include "AppleHttpBackgroundSession.h"
 #include "AppleHttp.h"
 #include "Apple/CFRef.h"
 #include "Misc/Base64.h"
 #include "Misc/CommandLine.h"
+#include "Misc/ScopeRWLock.h"
 
 #include "CommonCrypto/CommonDigest.h"
 
@@ -214,6 +217,16 @@ static bool ValidateCertificatePublicKeyPinningApple(NSURLAuthenticationChalleng
 
 #endif
 
+static TAutoConsoleVariable<bool> CVarEnableBackgroundRequests(
+	TEXT("AppleHttp.EnableBackgroundRequests"),
+	true,
+	TEXT("If true, will use the Background URLSession to schedule tasks that can run while the app is in the background. "
+			   "If false, will use the current flow scheduling the request to be executed."),
+	ECVF_Default);
+
+FRWLock FApplePlatformHttp::TaskRequestMapLock;
+TMap<NSURLSessionTask*, TSharedPtr<FAppleHttpRequest>> FApplePlatformHttp::TaskRequestMap;
+
 void FApplePlatformHttp::Init()
 {
 #if WITH_SSL
@@ -251,6 +264,23 @@ void FApplePlatformHttp::InitWithNSUrlSession()
 	
     [Session retain];
 
+#if PLATFORM_IOS // won't work on Mac, tvOS and visionOS
+	if (IsBackgroundRequestFeatureEnabled()) // can be disabled by cvar
+	{
+		Config = [NSURLSessionConfiguration backgroundSessionConfigurationWithIdentifier:@"com.unrealengine.backgroundhttpsession"];
+		
+		// Disable cache to mimic WinInet behavior
+		Config.requestCachePolicy = NSURLRequestReloadIgnoringLocalCacheData;
+		// Set discretionary to false so the OS will send the request ASAP (with fewer considerations on battery level, power consumption, and such)
+		Config.discretionary = NO;
+		Config.sessionSendsLaunchEvents = YES;
+
+		FApplePlatformHttpBackgroundSessionDelegate* Delegate = [[FApplePlatformHttpBackgroundSessionDelegate alloc] init];
+		BackgroundSession = [NSURLSession sessionWithConfiguration:Config delegate:Delegate delegateQueue:nil];
+
+		[BackgroundSession retain];
+	}
+#endif // PLATFORM_IOS
 }
 
 void FApplePlatformHttp::ShutdownWithNSUrlSession()
@@ -258,6 +288,10 @@ void FApplePlatformHttp::ShutdownWithNSUrlSession()
 	[Session invalidateAndCancel];
 	[Session release];
 	Session = nil;
+
+	[BackgroundSession invalidateAndCancel];
+	[BackgroundSession release];
+	BackgroundSession = nil;
 }
 
 FHttpManager* FApplePlatformHttp::CreatePlatformHttpManager()
@@ -272,6 +306,41 @@ IHttpRequest* FApplePlatformHttp::ConstructRequest()
 		InitWithNSUrlSession();
 	}
 
-	return new FAppleHttpRequest(Session);
+	// To avoid breaking the ConstructRequest contract I'll be passing both sessions here
+	// If the background feature is disabled, BackgroundSession will be nil and I can check that inside the FAppleHttpRequest
+	FAppleHttpRequest* Request = new FAppleHttpRequest(Session, BackgroundSession);
+
+	return Request;
+}
+
+bool FApplePlatformHttp::IsBackgroundRequestFeatureEnabled()
+{
+	return CVarEnableBackgroundRequests.GetValueOnAnyThread();
+}
+
+TSharedPtr<FAppleHttpRequest>& FApplePlatformHttp::Add(NSURLSessionTask* Task, const TSharedPtr<FAppleHttpRequest>& Request)
+{
+	FRWScopeLock ScopeLock(TaskRequestMapLock, FRWScopeLockType::SLT_Write);
+
+	return TaskRequestMap.Add(Task, CopyTemp(Request));
 }
 
+bool FApplePlatformHttp::Find(NSURLSessionTask* Task, TSharedPtr<FAppleHttpRequest>& Request)
+{
+	FRWScopeLock ScopeLock(TaskRequestMapLock, FRWScopeLockType::SLT_ReadOnly);
+
+	if (TSharedPtr<FAppleHttpRequest>* Found = TaskRequestMap.Find(Task))
+	{
+		Request = *Found;
+		return true;
+	}
+	
+	return false;
+}
+
+bool FApplePlatformHttp::RemoveAndCopyValue(NSURLSessionTask* Task, TSharedPtr<FAppleHttpRequest>& Request)
+{
+	FRWScopeLock ScopeLock(TaskRequestMapLock, FRWScopeLockType::SLT_Write);
+
+	return TaskRequestMap.RemoveAndCopyValue(Task, Request);
+}
\ No newline at end of file
diff --git a/Engine/Source/Runtime/Online/HTTP/Private/Interfaces/IHttpRequest.cpp b/Engine/Source/Runtime/Online/HTTP/Private/Interfaces/IHttpRequest.cpp
--- a/Engine/Source/Runtime/Online/HTTP/Private/Interfaces/IHttpRequest.cpp
+++ b/Engine/Source/Runtime/Online/HTTP/Private/Interfaces/IHttpRequest.cpp
@@ -2,6 +2,19 @@
 
 #include "Interfaces/IHttpRequest.h"
 
+const TCHAR* LexToString(EHttpRequestMode Mode)
+{
+	switch (Mode)
+	{
+		case EHttpRequestMode::ThreadedRequest: return TEXT("ThreadedRequest");
+		case EHttpRequestMode::ImmediateRequest: return TEXT("ImmediateRequest");
+
+		default:
+		checkf(false, TEXT("Please, implement new modes here"));
+		return TEXT("Unknown");
+	}
+}
+
 class FArchiveWithDelegateV2 final : public FArchive
 {
 public:
diff --git a/Engine/Source/Runtime/Online/HTTP/Public/Apple/ApplePlatformHttp.h b/Engine/Source/Runtime/Online/HTTP/Public/Apple/ApplePlatformHttp.h
--- a/Engine/Source/Runtime/Online/HTTP/Public/Apple/ApplePlatformHttp.h
+++ b/Engine/Source/Runtime/Online/HTTP/Public/Apple/ApplePlatformHttp.h
@@ -1,6 +1,7 @@
 // Copyright Epic Games, Inc. All Rights Reserved.
 
 #pragma once
+
 #include "GenericPlatform/GenericPlatformHttp.h"
 
 /**
@@ -9,7 +10,6 @@
 class HTTP_API FApplePlatformHttp : public FGenericPlatformHttp
 {
 public:
-
 	/**
 	 * Platform initialization step
 	 */
@@ -34,9 +34,19 @@ public:
 	 */
 	static IHttpRequest* ConstructRequest();
 
+	static bool IsBackgroundRequestFeatureEnabled();
+
+	static TSharedPtr<class FAppleHttpRequest>& Add(NSURLSessionTask* Task, const TSharedPtr<class FAppleHttpRequest>& Request);
+	static bool Find(NSURLSessionTask* Task, OUT TSharedPtr<class FAppleHttpRequest>& Request);
+	static bool RemoveAndCopyValue(NSURLSessionTask* Task, OUT TSharedPtr<class FAppleHttpRequest>& Request);
+
 private:
     /** Session used to create Apple based requests */
     static inline NSURLSession* Session = nil;
+	static inline NSURLSession* BackgroundSession = nil;
+
+	static FRWLock TaskRequestMapLock;
+	static TMap<NSURLSessionTask*, TSharedPtr<class FAppleHttpRequest>> TaskRequestMap;
 
 	static void InitWithNSUrlSession();
 	static void ShutdownWithNSUrlSession();
diff --git a/Engine/Source/Runtime/Online/HTTP/Public/HttpManager.h b/Engine/Source/Runtime/Online/HTTP/Public/HttpManager.h
--- a/Engine/Source/Runtime/Online/HTTP/Public/HttpManager.h
+++ b/Engine/Source/Runtime/Online/HTTP/Public/HttpManager.h
@@ -321,8 +321,6 @@ protected:
 	HTTP_API void UpdateUrlPatternsToDisableFailedLog(IConsoleVariable* CVar);
 	HTTP_API void UpdateUrlPatternsToMockResponse(IConsoleVariable* CVar);
 
-	HTTP_API void AddRequest(const FHttpRequestRef& Request);
-
 protected:
 	/** List of Http requests that are actively being processed */
 	TArray<FHttpRequestRef> Requests;
@@ -451,6 +449,13 @@ PACKAGE_SCOPE:
 	HTTP_API void MarkCurrentThreadCompletingRequest(bool bCompleting);
 	HTTP_API bool IsCurrentThreadCompletingRequest() const;
 
+	/**
+	 * Add a http request to the manager without enqueueing the threaded execution
+	 * 
+	 * @param Request - the request object to add
+	 */
+	HTTP_API void AddRequest(const FHttpRequestRef& Request);
+
 protected:
 	TMap<FString, FMockResponse> UrlPatternsToMockResponse;
 	mutable FCriticalSection UrlPatternsToMockResponseCriticalSection;
diff --git a/Engine/Source/Runtime/Online/HTTP/Public/Interfaces/IHttpRequest.h b/Engine/Source/Runtime/Online/HTTP/Public/Interfaces/IHttpRequest.h
--- a/Engine/Source/Runtime/Online/HTTP/Public/Interfaces/IHttpRequest.h
+++ b/Engine/Source/Runtime/Online/HTTP/Public/Interfaces/IHttpRequest.h
@@ -131,11 +131,25 @@ private:
 namespace HttpRequestOptions
 {
 	static const FName HttpVersion("HttpVersion");
+	static const FName RequestMode("RequestMode");
 #if UE_HTTP_SUPPORT_UNIX_SOCKET
 	static const FName UnixSocketPath("UnixSocketPath");
 #endif //UE_HTTP_SUPPORT_UNIX_SOCKET
 }
 
+/**
+ * Mode that can be specified through IHttpRequest::SetOption(HttpRequestOptions::RequestMode, LexToString(EHttpRequestMode::Value));
+ */
+enum class EHttpRequestMode : uint8
+{
+	/** The request will be enqueued to the http manager thread system */
+	ThreadedRequest,
+	/** The request will be immediately sent when the platform supports it */
+	ImmediateRequest
+};
+
+HTTP_API const TCHAR* LexToString(EHttpRequestMode Mode);
+
 /**
  * Interface for Http requests (created using FHttpFactory)
  */
diff --git a/Engine/Source/Runtime/StudioTelemetry/Private/AnalyticsProviderBroadcast.cpp b/Engine/Source/Runtime/StudioTelemetry/Private/AnalyticsProviderBroadcast.cpp
--- a/Engine/Source/Runtime/StudioTelemetry/Private/AnalyticsProviderBroadcast.cpp
+++ b/Engine/Source/Runtime/StudioTelemetry/Private/AnalyticsProviderBroadcast.cpp
@@ -2,6 +2,7 @@
 
 #include "AnalyticsProviderBroadcast.h"
 #include "Analytics.h"
+#include "Interfaces/IAnalyticsProvider.h"
 #include "Interfaces/IAnalyticsProviderModule.h"
 #include "AnalyticsProviderConfigurationDelegate.h"
 #include "Misc/ConfigCacheIni.h"
@@ -282,14 +283,14 @@ static void CheckForDuplicateAttributes(const FString& EventName, const TArray<F
 #endif
 }
 
-void FAnalyticsProviderBroadcast::RecordEvent(const FString& EventName, const TArray<FAnalyticsEventAttribute>& Attributes)
+void FAnalyticsProviderBroadcast::RecordEvent(const FString& EventName, const TArray<FAnalyticsEventAttribute>& Attributes, EAnalyticsRecordEventMode Mode)
 {
 	// Check for duplicate attributes against themselves
 	CheckForDuplicateAttributes(EventName, Attributes);
 
 	for (TProviders::TConstIterator it(Providers); it; ++it)
 	{
-		(*it).Value->RecordEvent(EventName, Attributes);
+		(*it).Value->RecordEvent(EventName, Attributes, Mode);
 	}
 
 	if (RecordEventCallback)
diff --git a/Engine/Source/Runtime/StudioTelemetry/Private/AnalyticsProviderBroadcast.h b/Engine/Source/Runtime/StudioTelemetry/Private/AnalyticsProviderBroadcast.h
--- a/Engine/Source/Runtime/StudioTelemetry/Private/AnalyticsProviderBroadcast.h
+++ b/Engine/Source/Runtime/StudioTelemetry/Private/AnalyticsProviderBroadcast.h
@@ -49,7 +49,12 @@ public:
 	virtual FString GetSessionID() const override;
 	virtual FString GetUserID() const override;
 	
-	virtual void RecordEvent(const FString& EventName, const TArray<FAnalyticsEventAttribute>& Attributes) override;	
+	virtual void RecordEvent(const FString& EventName, const TArray<FAnalyticsEventAttribute>& Attributes) override 
+	{
+		RecordEvent(EventName, Attributes, EAnalyticsRecordEventMode::Cached);
+	}
+	
+	virtual void RecordEvent(const FString& EventName, const TArray<FAnalyticsEventAttribute>& Attributes, EAnalyticsRecordEventMode Mode) override;	
 	void SetRecordEventCallback(OnRecordEvent Callback);
 	bool HasValidProviders() const { return Providers.Num() > 0; }
 	TWeakPtr<IAnalyticsProvider> GetAnalyticsProvider(const FString& Name);
