[Backout] - CL44527427
#rnx
Original CL Desc
-----------------------------------------------------------------
Update the mobile funnel telemetry to include an event that fires as soon as download is 100% complete.

- Moving the iOS funnel event to a background URLSession

Detailed explanation:
The main goal was to be able to send an analytics event immediately, instead of queueing it to be sent using the threaded system we have nowadays (for iOS).
The current system does not work since the game stops ticking when backgrounded.
So, we took advantage of the iOS background URLSession that supports downloads and uploads.
Making it possible to schedule the upload task with the analytics data as the body (memory file).
This change opens a new pathway for platforms to implement support for making requests that are system handled (or any other custom implementation), while in the background.
So the first use case is for sending analytics events on iOS (the only platform implementing it ATM), but we could research adding more platforms to it when supported.

#rnx

diff --git a/Engine/Source/Runtime/Analytics/Analytics/Private/Interfaces/IAnalyticsProvider.cpp b/Engine/Source/Runtime/Analytics/Analytics/Private/Interfaces/IAnalyticsProvider.cpp
deleted file mode 100644
diff --git a/Engine/Source/Runtime/Analytics/Analytics/Public/Interfaces/IAnalyticsProvider.h b/Engine/Source/Runtime/Analytics/Analytics/Public/Interfaces/IAnalyticsProvider.h
--- a/Engine/Source/Runtime/Analytics/Analytics/Public/Interfaces/IAnalyticsProvider.h
+++ b/Engine/Source/Runtime/Analytics/Analytics/Public/Interfaces/IAnalyticsProvider.h
@@ -2,24 +2,11 @@
 
 #pragma once
 
-#include "CoreMinimal.h"
 #include "AnalyticsEventAttribute.h"
+#include "CoreMinimal.h"
 
 class Error;
 
-/**
- * Enum describing the supported record event modes
- * @enum Cached use the existing system where the events are cached and flushed when possible
- * @enum Immediate use the new system where events are immediately flushed to a HTTP request
- */
-enum class EAnalyticsRecordEventMode : uint8
-{
-	Cached,
-	Immediate
-};
-
-const TCHAR* LexToString(EAnalyticsRecordEventMode Mode);
-
 /** 
  * Generic interface for an analytics provider. 
  * Other modules can define more and register them with this module. See FAnalytics for details.
@@ -203,18 +190,6 @@ public:
 	 */
 	virtual void RecordEvent(const FString& EventName, const TArray<FAnalyticsEventAttribute>& Attributes) = 0;
 
-	/**
-	 * RecordEvent overload to allow platforms to support events that are immediately sent
-	 *
-	 * @param EventName name of the event
-	 * @param Attributes array of attribute name/value pairs
-	 * @param Mode the record event mode, Cached or Immediate @see EAnalyticsRecordEventMode
-	 */
-	virtual void RecordEvent(const FString& EventName, const TArray<FAnalyticsEventAttribute>& Attributes, EAnalyticsRecordEventMode Mode)
-	{
-		RecordEvent(EventName, Attributes);
-	}
-
 	/**
 	 * Overload for RecordEvent that takes no parameters
 	 *
diff --git a/Engine/Source/Runtime/Analytics/AnalyticsET/Private/AnalyticsProviderETEventCache.cpp b/Engine/Source/Runtime/Analytics/AnalyticsET/Private/AnalyticsProviderETEventCache.cpp
--- a/Engine/Source/Runtime/Analytics/AnalyticsET/Private/AnalyticsProviderETEventCache.cpp
+++ b/Engine/Source/Runtime/Analytics/AnalyticsET/Private/AnalyticsProviderETEventCache.cpp
@@ -346,14 +346,6 @@ TArray<uint8> FAnalyticsProviderETEventCache::FlushCacheUTF8()
 	return TArray<uint8>();
 }
 
-TArray<uint8> FAnalyticsProviderETEventCache::CreateImmediatePayload(const FString& EventName, const TArray<FAnalyticsEventAttribute>& Attributes) const
-{
-	TSharedRef<FAnalyticsProviderETEventCache> Clone = CloneForImmediateMode();
-	Clone->AddToCache(EventName, Attributes);
-	return Clone->FlushCacheUTF8();
-	// the Clone will be released here
-}
-
 // !!!!!!!!!!!!!!!!!!!!!!!!!!!!!!!!!!!!!!!!!!!!!!!!!!!!!!!!!!!!!!!!!!!!!!
 // !!!! This method tries extremely hard to avoid any dynamic allocations
 // !!!! to optimize the flush time. Please don't add new allocations to this function 
@@ -433,14 +425,6 @@ void FAnalyticsProviderETEventCache::QueueFlush()
 	EventCacheStatic::InitializePayloadBuffer(CachedEventUTF8Stream, PreallocatedPayloadSize);
 }
 
-TSharedRef<FAnalyticsProviderETEventCache> FAnalyticsProviderETEventCache::CloneForImmediateMode() const
-{
-	// need to make this because FAnalyticsProviderETEventCache is not allowed to be copyed because
-	// FTransactionallySafeCriticalSection CachedEventsCS depends on FPlatformSharedMutex which has deleted copy and assignment operators.
-	TSharedRef<FAnalyticsProviderETEventCache> Copy = MakeShared<FAnalyticsProviderETEventCache>(MaximumPayloadSize, PreallocatedPayloadSize);
-	Copy->CachedDefaultAttributeUTF8Stream = CachedDefaultAttributeUTF8Stream;
-	return Copy;
-}
 
 bool FAnalyticsProviderETEventCache::CanFlush() const
 {
diff --git a/Engine/Source/Runtime/Analytics/AnalyticsET/Private/IAnalyticsProviderET.cpp b/Engine/Source/Runtime/Analytics/AnalyticsET/Private/IAnalyticsProviderET.cpp
--- a/Engine/Source/Runtime/Analytics/AnalyticsET/Private/IAnalyticsProviderET.cpp
+++ b/Engine/Source/Runtime/Analytics/AnalyticsET/Private/IAnalyticsProviderET.cpp
@@ -143,7 +143,6 @@ public:
 
 	virtual bool ShouldRecordEvent(const FString& EventName) const override;
 	virtual void RecordEvent(FString&& EventName, const TArray<FAnalyticsEventAttribute>& Attributes) override;
-	virtual void RecordEvent(FString&& EventName, const TArray<FAnalyticsEventAttribute>& Attributes, EAnalyticsRecordEventMode Mode) override;
 	virtual void SetDefaultEventAttributes(TArray<FAnalyticsEventAttribute>&& Attributes) override;
 	virtual TArray<FAnalyticsEventAttribute> GetDefaultEventAttributesSafe() const override;
 	virtual int32 GetDefaultEventAttributeCount() const override;
@@ -162,8 +161,6 @@ public:
 	virtual const FAnalyticsET::Config& GetConfig() const override { return Config; }
 
 private:
-	void ExecuteRequest(TArray<uint8>& Payload, OUT int32& PayloadSize, OUT int32& EventCount, EAnalyticsRecordEventMode Mode);
-	void SendImmediately(const FString& EventName, const TArray<FAnalyticsEventAttribute>& Attributes);
 	void FlushEventsOnce();
 	void FlushEventLegacy(const FString& EventName, const TArray<FAnalyticsEventAttribute>& Attributes);
 	bool IsActingAsNullProvider() const
@@ -174,7 +171,7 @@ private:
 	}
 
 	/** Create a request utilizing HttpRetry domains */
-	TSharedRef<IHttpRequest, ESPMode::ThreadSafe> CreateRequest(EAnalyticsRecordEventMode Mode = EAnalyticsRecordEventMode::Cached);
+	TSharedRef<IHttpRequest, ESPMode::ThreadSafe> CreateRequest();
 
 	bool bSessionInProgress;
 	/** The current configuration (might be updated with respect to the one provided at construction). */
@@ -425,7 +422,7 @@ void FAnalyticsProviderET::EndSession()
 	bSessionInProgress = false;
 }
 
-TSharedRef<IHttpRequest, ESPMode::ThreadSafe> FAnalyticsProviderET::CreateRequest(EAnalyticsRecordEventMode Mode)
+TSharedRef<IHttpRequest, ESPMode::ThreadSafe> FAnalyticsProviderET::CreateRequest()
 {
 	LLM_SCOPE_BYNAME(TEXT("Analytics"));
 	if (!ensure(FModuleManager::Get().IsModuleLoaded("HTTP")))
@@ -438,19 +435,12 @@ TSharedRef<IHttpRequest, ESPMode::ThreadSafe> FAnalyticsProviderET::CreateReques
 		FHttpRetrySystem::FRetryTimeoutRelativeSecondsSetting(),
 		FHttpRetrySystem::FRetryResponseCodes(),
 		FHttpRetrySystem::FRetryVerbs(),
-		RetryServers,
-		FHttpRetrySystem::FRetryLimitCountSetting(),
-		FHttpRetrySystem::FExponentialBackoffCurve());
+		RetryServers);
 	for (const TPair<FString, FString>& HttpHeader : HttpHeaders)
 	{
 		HttpRequest->SetHeader(HttpHeader.Key, HttpHeader.Value);
 	}
 
-	if (Mode == EAnalyticsRecordEventMode::Immediate)
-	{
-		HttpRequest->SetOption(HttpRequestOptions::RequestMode, LexToString(EHttpRequestMode::ImmediateRequest));
-	}
-
 	return HttpRequest;
 }
 
@@ -470,86 +460,6 @@ void FAnalyticsProviderET::FlushEvents()
 	};
 }
 
-void FAnalyticsProviderET::ExecuteRequest(TArray<uint8>& Payload, OUT int32& PayloadSize, OUT int32& EventCount, EAnalyticsRecordEventMode Mode)
-{
-	EventCount = 0;
-	PayloadSize = 0;
-
-	// UrlEncode NOTE: need to concatenate everything
-	FString URLPath = Config.APIEndpointET;
-			URLPath += TEXT("?SessionID=") + FPlatformHttp::UrlEncode(SessionID);
-			URLPath += TEXT("&AppID=") + FPlatformHttp::UrlEncode(Config.APIKeyET);
-			URLPath += TEXT("&AppVersion=") + FPlatformHttp::UrlEncode(Config.AppVersionET);
-			URLPath += TEXT("&UserID=") + FPlatformHttp::UrlEncode(UserID);
-			URLPath += TEXT("&AppEnvironment=") + FPlatformHttp::UrlEncode(Config.AppEnvironment);
-			URLPath += TEXT("&UploadType=") + FPlatformHttp::UrlEncode(Config.UploadType);
-	PayloadSize = URLPath.Len() + Payload.Num();
-
-	// Recreate the URLPath for logging because we do not want to escape the parameters when logging.
-	// We cannot simply UrlEncode the entire Path after logging it because UrlEncode(Params) != UrlEncode(Param1) & UrlEncode(Param2) ...
-	UE_LOG(LogAnalytics, VeryVerbose, TEXT("[%s] AnalyticsET URL:%s?SessionID=%s&AppID=%s&AppVersion=%s&UserID=%s&AppEnvironment=%s&UploadType=%s. Payload:%hs"),
-		*Config.APIKeyET,
-		*Config.APIEndpointET,
-		*SessionID,
-		*Config.APIKeyET,
-		*Config.AppVersionET,
-		*UserID,
-		*Config.AppEnvironment,
-		*Config.UploadType,
-		FUtf8StringView((FGenericPlatformTypes::UTF8CHAR*)Payload.GetData(), Payload.Num()).GetData());
-
-	{
-		QUICK_SCOPE_CYCLE_COUNTER(STAT_FlushEventsHttpRequest);
-
-		// Create/send Http request for an event
-		TSharedRef<IHttpRequest, ESPMode::ThreadSafe> HttpRequest = CreateRequest(Mode);
-		HttpRequest->SetHeader(TEXT("Content-Type"), TEXT("application/json; charset=utf-8"));
-		// Want to avoid putting the project name into the User-Agent, because for some apps (like the editor), the project name is private info.
-		// The analytics User-Agent uses the default User-Agent, but with project name removed.
-		HttpRequest->SetHeader(TEXT("User-Agent"), UserAgentCache.GetUserAgent());
-		HttpRequest->SetURL(Config.APIServerET / URLPath);
-		HttpRequest->SetVerb(TEXT("POST"));
-		HttpRequest->SetContent(MoveTemp(Payload));
-
-#if UE_HTTP_SUPPORT_UNIX_SOCKET
-		if (!Config.APIUnixSocketPathET.IsEmpty())
-		{
-			HttpRequest->SetOption(HttpRequestOptions::UnixSocketPath, Config.APIUnixSocketPathET);
-		}
-#endif //UE_HTTP_SUPPORT_UNIX_SOCKET
-
-		// Don't set a response callback if we are in our destructor, as the instance will no longer be there to call.
-		if (!bInDestructor)
-		{
-			HttpRequest->OnProcessRequestComplete().BindSP(this, &FAnalyticsProviderET::EventRequestComplete);
-		}
-
-		OnPreAnalyticsEventProcessedDelegate.Broadcast(HttpRequest);
-
-		HttpRequest->ProcessRequest();
-	}
-}
-
-void FAnalyticsProviderET::SendImmediately(const FString& EventName, const TArray<FAnalyticsEventAttribute>& Attributes)
-{
-	LLM_SCOPE_BYNAME(TEXT("Analytics"));
-
-	if (ensure(FModuleManager::Get().IsModuleLoaded("HTTP")))
-	{
-		TArray<uint8> Payload = EventCache.CreateImmediatePayload(EventName, Attributes);
-
-		int32 PayloadSize = 0;
-		int32 EventCount = 0;
-		ExecuteRequest(Payload, PayloadSize, EventCount, EAnalyticsRecordEventMode::Immediate);
-	}
-	else
-	{
-		// If the HTTP module is not loaded yet for some reason, we fallback to the Cached system so we won't lose this event.
-		// This should not happen, but better safe than sorry.
-		RecordEvent(CopyTemp(EventName), Attributes, EAnalyticsRecordEventMode::Cached);
-	}
-}
-
 void FAnalyticsProviderET::FlushEventsOnce()
 {
 	// FlushEventsOnce cannot be rolled back, so it must not occur inside an AutoRTFM transaction.
@@ -572,8 +482,66 @@ void FAnalyticsProviderET::FlushEventsOnce()
 	if(ensure(FModuleManager::Get().IsModuleLoaded("HTTP")))
 	{
 		TArray<uint8> Payload = EventCache.FlushCacheUTF8();
-		
-		ExecuteRequest(Payload, PayloadSize, EventCount, EAnalyticsRecordEventMode::Cached);
+		// UrlEncode NOTE: need to concatenate everything
+		FString URLPath = Config.APIEndpointET;
+				URLPath += TEXT("?SessionID=") + FPlatformHttp::UrlEncode(SessionID);
+				URLPath += TEXT("&AppID=") + FPlatformHttp::UrlEncode(Config.APIKeyET);
+				URLPath += TEXT("&AppVersion=") + FPlatformHttp::UrlEncode(Config.AppVersionET);
+				URLPath += TEXT("&UserID=") + FPlatformHttp::UrlEncode(UserID);
+				URLPath += TEXT("&AppEnvironment=") + FPlatformHttp::UrlEncode(Config.AppEnvironment);
+				URLPath += TEXT("&UploadType=") + FPlatformHttp::UrlEncode(Config.UploadType);
+		PayloadSize = URLPath.Len() + Payload.Num();
+
+		// This should never be done in production. MUCH slower!
+		if (UE_LOG_ACTIVE(LogAnalytics, VeryVerbose))
+		{
+			// need to null terminate to load the payload.
+			Payload.Add(TEXT('\0'));
+			// Recreate the URLPath for logging because we do not want to escape the parameters when logging.
+			// We cannot simply UrlEncode the entire Path after logging it because UrlEncode(Params) != UrlEncode(Param1) & UrlEncode(Param2) ...
+			FString LogString = FString::Printf(TEXT("[%s] AnalyticsET URL:%s?SessionID=%s&AppID=%s&AppVersion=%s&UserID=%s&AppEnvironment=%s&UploadType=%s. Payload:%s"),
+				*Config.APIKeyET,
+				*Config.APIEndpointET,
+				*SessionID,
+				*Config.APIKeyET,
+				*Config.AppVersionET,
+				*UserID,
+				*Config.AppEnvironment,
+				*Config.UploadType,
+				UTF8_TO_TCHAR(Payload.GetData()));
+			UE_LOG(LogAnalytics, VeryVerbose, TEXT("%s"), *LogString);
+			Payload.SetNum(Payload.Num()-1);
+		}
+
+		{
+			QUICK_SCOPE_CYCLE_COUNTER(STAT_FlushEventsHttpRequest);
+			// Create/send Http request for an event
+			TSharedRef<IHttpRequest, ESPMode::ThreadSafe> HttpRequest = CreateRequest();
+			HttpRequest->SetHeader(TEXT("Content-Type"), TEXT("application/json; charset=utf-8"));
+			// Want to avoid putting the project name into the User-Agent, because for some apps (like the editor), the project name is private info.
+			// The analytics User-Agent uses the default User-Agent, but with project name removed.
+			HttpRequest->SetHeader(TEXT("User-Agent"), UserAgentCache.GetUserAgent());
+			HttpRequest->SetURL(Config.APIServerET / URLPath);
+			HttpRequest->SetVerb(TEXT("POST"));
+			HttpRequest->SetContent(MoveTemp(Payload));
+
+#if UE_HTTP_SUPPORT_UNIX_SOCKET
+			if (!Config.APIUnixSocketPathET.IsEmpty())
+			{
+				HttpRequest->SetOption(HttpRequestOptions::UnixSocketPath, Config.APIUnixSocketPathET);
+			}
+#endif //UE_HTTP_SUPPORT_UNIX_SOCKET
+
+			// Don't set a response callback if we are in our destructor, as the instance will no longer be there to call.
+			if (!bInDestructor)
+			{
+				HttpRequest->OnProcessRequestComplete().BindSP(this, &FAnalyticsProviderET::EventRequestComplete);
+			}
+
+			OnPreAnalyticsEventProcessedDelegate.Broadcast(HttpRequest);
+
+			HttpRequest->ProcessRequest();
+		}
 	}
 	ANALYTICS_FLUSH_TRACKING_END(PayloadSize, EventCount);
 }
@@ -651,51 +619,29 @@ void FAnalyticsProviderET::RecordEvent(FString&& EventName, const TArray<FAnalyt
 {
 	LLM_SCOPE_BYNAME(TEXT("Analytics"));
 
-	RecordEvent(MoveTemp(EventName), Attributes, EAnalyticsRecordEventMode::Cached);
-}
-
-void FAnalyticsProviderET::RecordEvent(FString&& EventName, const TArray<FAnalyticsEventAttribute>& Attributes, EAnalyticsRecordEventMode Mode)
-{
 	// let higher level code filter the decision of whether to send the event
 	if (ShouldRecordEvent(EventName))
 	{
-		switch (Mode)
+		// fire any callbacks
+		for (const auto& Cb : EventRecordedCallbacks)
 		{
-			case EAnalyticsRecordEventMode::Cached:
-			{
-				// fire any callbacks
-				for (const auto& Cb : EventRecordedCallbacks)
-				{
-					// we no longer track if the event was Json, each attribute does.
-					Cb(EventName, Attributes, false);
-				}
-
-				if (!Config.UseLegacyProtocol)
-				{
-					EventCache.AddToCache(MoveTemp(EventName), Attributes);
-
-					// if we aren't caching events, flush immediately. This is really only for debugging as it will significantly affect bandwidth.
-					if (!bShouldCacheEvents)
-					{
-						FlushEvents();
-					}
-				}
-				else
-				{
-					FlushEventLegacy(EventName, Attributes);
-				}
-			}
-			break;
+			// we no longer track if the event was Json, each attribute does.
+			Cb(EventName, Attributes, false);
+		}
+
+		if (!Config.UseLegacyProtocol)
+		{
+			EventCache.AddToCache(MoveTemp(EventName), Attributes);
 
-			case EAnalyticsRecordEventMode::Immediate:
+			// if we aren't caching events, flush immediately. This is really only for debugging as it will significantly affect bandwidth.
+			if (!bShouldCacheEvents)
 			{
-				SendImmediately(MoveTemp(EventName), Attributes);
+				FlushEvents();
 			}
-			break;
-
-			default:
-			checkf(false, TEXT("Please, implement new modes here"));
-			break;
+		}
+		else
+		{
+			FlushEventLegacy(EventName, Attributes);
 		}
 	}
 	else
diff --git a/Engine/Source/Runtime/Analytics/AnalyticsET/Public/AnalyticsProviderETEventCache.h b/Engine/Source/Runtime/Analytics/AnalyticsET/Public/AnalyticsProviderETEventCache.h
--- a/Engine/Source/Runtime/Analytics/AnalyticsET/Public/AnalyticsProviderETEventCache.h
+++ b/Engine/Source/Runtime/Analytics/AnalyticsET/Public/AnalyticsProviderETEventCache.h
@@ -71,11 +71,6 @@ public:
 	/** Flushes the cache as a UTF8 char array. Returns a uint8 because that's what IHttpRequest prefers. If nothing is cached, returns an empty array. */
 	ANALYTICSET_API TArray<uint8> FlushCacheUTF8();
 
-	/**
-	 * Creates a immediate payload with the existing default attributes
-	 */
-	ANALYTICSET_API TArray<uint8> CreateImmediatePayload(const FString& EventName, const TArray<FAnalyticsEventAttribute>& Attributes) const;
-
 	/**
 	* Determines whether we have anything we need to flush, either a queued flush or existing events in the payload.
 	*/
@@ -120,8 +115,6 @@ public:
 private:
 	ANALYTICSET_API void QueueFlush();
 
-	ANALYTICSET_API TSharedRef<FAnalyticsProviderETEventCache> CloneForImmediateMode() const;
-
 	/**
 	* Analytics event entry to be cached
 	*/
diff --git a/Engine/Source/Runtime/Analytics/AnalyticsET/Public/IAnalyticsProviderET.h b/Engine/Source/Runtime/Analytics/AnalyticsET/Public/IAnalyticsProviderET.h
--- a/Engine/Source/Runtime/Analytics/AnalyticsET/Public/IAnalyticsProviderET.h
+++ b/Engine/Source/Runtime/Analytics/AnalyticsET/Public/IAnalyticsProviderET.h
@@ -37,14 +37,6 @@ public:
 		return RecordEvent(CopyTemp(EventName), Attributes);
 	}
 
-	/**
-	 * This class augments RecordEvent with a version that takes the EventName by rvalue reference to save a string copy. Implement the base version in terms of this one.
-	 */
-	virtual void RecordEvent(const FString& EventName, const TArray<FAnalyticsEventAttribute>& Attributes, EAnalyticsRecordEventMode Mode) override
-	{
-		return RecordEvent(CopyTemp(EventName), Attributes, Mode);
-	}
-
 	/**
 	 * This class augments StartSession with a version that takes the SessionID instead of always generating it. Implement the base version in terms of this one.
 	 */
@@ -104,11 +96,6 @@ public:
 	 * Primary RecordEvent API. Allow move semantics to capture the attributes.
 	 */
 	virtual void RecordEvent(FString&& EventName, const TArray<FAnalyticsEventAttribute>& Attributes) = 0;
-	
-	/**
-	 * Overload to allow platforms to implement events that are immediately sent when supported. @see EAnalyticsRecordEventMode.
-	 */
-	virtual void RecordEvent(FString&& EventName, const TArray<FAnalyticsEventAttribute>& Attributes, EAnalyticsRecordEventMode Mode) = 0;
 
 	/**
 	 * Updates the default Domain and AltDomains.
diff --git a/Engine/Source/Runtime/ApplicationCore/Private/IOS/IOSAppDelegate.cpp b/Engine/Source/Runtime/ApplicationCore/Private/IOS/IOSAppDelegate.cpp
--- a/Engine/Source/Runtime/ApplicationCore/Private/IOS/IOSAppDelegate.cpp
+++ b/Engine/Source/Runtime/ApplicationCore/Private/IOS/IOSAppDelegate.cpp
@@ -263,6 +263,7 @@ bool FIOSCoreDelegates::PassesPushNotificationFilters(NSDictionary* Payload)
 @synthesize AccessibilityCacheTimer;
 #endif
 @synthesize savedOpenUrlParameters;
+@synthesize BackgroundSessionEventCompleteDelegate;
 
 
 static IOSAppDelegate* CachedDelegate = nil;
@@ -1871,7 +1872,7 @@ extern double GCStartTime;
 - (void)application:(UIApplication *)application handleEventsForBackgroundURLSession:(NSString *)identifier completionHandler:(void(^)(void))completionHandler
 {
     //Save off completionHandler so that a future call to FCoreDelegates::ApplicationBackgroundSessionEventsAllSentDelegate can execute it
-	[FBackgroundURLSessionHandler::BackgroundSessionEventCompleteDelegateMap setObject:completionHandler forKey:identifier];
+    self.BackgroundSessionEventCompleteDelegate = completionHandler;
     
     // Invoke background session with this identifier if needed to handle these events
 	FString Id(identifier);
diff --git a/Engine/Source/Runtime/ApplicationCore/Private/IOS/IOSBackgroundURLSessionHandler.cpp b/Engine/Source/Runtime/ApplicationCore/Private/IOS/IOSBackgroundURLSessionHandler.cpp
--- a/Engine/Source/Runtime/ApplicationCore/Private/IOS/IOSBackgroundURLSessionHandler.cpp
+++ b/Engine/Source/Runtime/ApplicationCore/Private/IOS/IOSBackgroundURLSessionHandler.cpp
@@ -1427,15 +1427,14 @@ static constexpr NSInteger HTTPStatusCodeErrorServer = 500;
 {
 	UE_DNLD_LOG(@"URLSessionDidFinishEventsForBackgroundURLSession");
 
-	NSString* Id = Session.configuration.identifier;
-
-	if ([FBackgroundURLSessionHandler::BackgroundSessionEventCompleteDelegateMap objectForKey:Id] == nil)
+	IOSAppDelegate* AppDelegate = [IOSAppDelegate GetDelegate];
+	if (AppDelegate == nullptr || AppDelegate.BackgroundSessionEventCompleteDelegate == nil)
 	{
 		return;
 	}
 	
-	void(^CompletionHandler)() = [[FBackgroundURLSessionHandler::BackgroundSessionEventCompleteDelegateMap objectForKey:Id] retain];
-	[FBackgroundURLSessionHandler::BackgroundSessionEventCompleteDelegateMap removeObjectForKey:Id];
+	void(^CompletionHandler)() = [AppDelegate.BackgroundSessionEventCompleteDelegate retain];
+	AppDelegate.BackgroundSessionEventCompleteDelegate = nil;
 
 	// Completion handler has to be invoked on the main thread.
 	[[NSOperationQueue mainQueue] addOperationWithBlock:^
@@ -1648,8 +1647,6 @@ FBackgroundURLSessionHandler::FOnDownloadMetrics FBackgroundURLSessionHandler::O
 
 FBackgroundURLSessionHandler::FOnDownloadsCompletedWhileAppWasNotRunning FBackgroundURLSessionHandler::OnDownloadsCompletedWhileAppWasNotRunning;
 
-NSMutableDictionary<NSString*,void(^)()>* FBackgroundURLSessionHandler::BackgroundSessionEventCompleteDelegateMap = [[NSMutableDictionary alloc] init];
-
 void FBackgroundURLSessionHandler::AllowCellular(bool bAllow)
 {
 	@autoreleasepool
diff --git a/Engine/Source/Runtime/ApplicationCore/Public/IOS/IOSAppDelegate.h b/Engine/Source/Runtime/ApplicationCore/Public/IOS/IOSAppDelegate.h
--- a/Engine/Source/Runtime/ApplicationCore/Public/IOS/IOSAppDelegate.h
+++ b/Engine/Source/Runtime/ApplicationCore/Public/IOS/IOSAppDelegate.h
@@ -302,6 +302,8 @@ APPLICATIONCORE_API
 
 - (void)ForceExit;
 
+@property (nonatomic, copy) void(^BackgroundSessionEventCompleteDelegate)();
+
 static void interruptionListener(void* ClientData, UInt32 Interruption);
 
 -(UIWindow*)window;
diff --git a/Engine/Source/Runtime/ApplicationCore/Public/IOS/IOSBackgroundURLSessionHandler.h b/Engine/Source/Runtime/ApplicationCore/Public/IOS/IOSBackgroundURLSessionHandler.h
--- a/Engine/Source/Runtime/ApplicationCore/Public/IOS/IOSBackgroundURLSessionHandler.h
+++ b/Engine/Source/Runtime/ApplicationCore/Public/IOS/IOSBackgroundURLSessionHandler.h
@@ -70,9 +70,6 @@ public:
 	DECLARE_MULTICAST_DELEGATE_OneParam(FOnDownloadsCompletedWhileAppWasNotRunning, const bool /*bSuccess*/);
 	static FOnDownloadsCompletedWhileAppWasNotRunning OnDownloadsCompletedWhileAppWasNotRunning;
 
-	// Will be used by IOSAppDelegate and ApplePlatformHttp
-	static NSMutableDictionary<NSString*,void(^)()>* BackgroundSessionEventCompleteDelegateMap;
-
 	// Sets if cellular is allowed to be used for new downloads.
 	// Existing downloads will be recreated to reflect new setting value.
 	static void AllowCellular(bool bAllow);
diff --git a/Engine/Source/Runtime/Online/HTTP/HTTP.Build.cs b/Engine/Source/Runtime/Online/HTTP/HTTP.Build.cs
--- a/Engine/Source/Runtime/Online/HTTP/HTTP.Build.cs
+++ b/Engine/Source/Runtime/Online/HTTP/HTTP.Build.cs
@@ -86,16 +86,6 @@ public class HTTP : ModuleRules
 			}
 			);
 
-		if (Target.bCompileAgainstApplicationCore && Target.Platform == UnrealTargetPlatform.IOS)
-		{
-			PrivateDependencyModuleNames.AddRange(
-				new string[] {
-					"ApplicationCore",
-					"BackgroundHTTPFileHash",
-				}
-			);
-		}
-
 		PrivateDependencyModuleNames.AddRange(
 			new string[] {
 				"EventLoop",
diff --git a/Engine/Source/Runtime/Online/HTTP/Private/Apple/AppleHttp.cpp b/Engine/Source/Runtime/Online/HTTP/Private/Apple/AppleHttp.cpp
--- a/Engine/Source/Runtime/Online/HTTP/Private/Apple/AppleHttp.cpp
+++ b/Engine/Source/Runtime/Online/HTTP/Private/Apple/AppleHttp.cpp
@@ -2,8 +2,6 @@
 
 
 #include "AppleHttp.h"
-#include "Apple/ApplePlatformHttp.h"
-#include "AppleHttpObjc.h"
 #include "HAL/FileManager.h"
 #include "HAL/IConsoleManager.h"
 #include "HAL/PlatformFileManager.h"
@@ -11,11 +9,62 @@
 #include "Http.h"
 #include "HttpManager.h"
 #include "HttpModule.h"
-#include "Interfaces/IHttpRequest.h"
 #include "Misc/EngineVersion.h"
 #include "Misc/App.h"
 #include "Misc/Base64.h"
 
+// UE-228716
+TAutoConsoleVariable<bool> CVarAppleHttpStartTaskInHttpThreadEnabled(
+	TEXT("AppleHttp.StartTaskInHttpThreadEnabled"),
+	true,
+	TEXT("Enables this implementation to start the task in the http thread like the rest of implementations"),
+	ECVF_SaveForNextBoot
+);
+
+/**
+ * Class to hold data from delegate implementation notifications.
+ */
+
+@interface FAppleHttpResponseDelegate : NSObject<NSURLSessionDataDelegate>
+{
+	// Flag to indicate the request was initialized with stream. In that case even if stream was set to 
+	// null later on internally, the request itself won't cache received data anymore
+	@public BOOL bInitializedWithValidStream;
+
+	/** Have we received any data? */
+	BOOL bAnyHttpActivity;
+
+	/** Delegate invoked after processing URLSession:dataTask:didReceiveData or URLSession:task:didCompleteWithError:*/
+	@public FNewAppleHttpEventDelegate NewAppleHttpEventDelegate;
+}
+
+/** A handle for the response */
+@property(retain) NSURLResponse* Response;
+/** The total number of bytes written out during the request/response */
+@property uint64 BytesWritten;
+/** The total number of bytes received out during the request/response */
+@property uint64 BytesReceived;
+/** Request status */
+@property EHttpRequestStatus::Type RequestStatus;
+/** Reason of failure */
+@property EHttpFailureReason FailureReason;
+/** Associated request. Cleared when canceled */
+@property TWeakPtr<FAppleHttpRequest> SourceRequest;
+
+/** NSURLSessionDataDelegate delegate methods. Those are called from a thread controlled by the NSURLSession */
+
+/** Sent periodically to notify the delegate of upload progress. */
+- (void)URLSession:(NSURLSession *)session task:(NSURLSessionTask *)task didSendBodyData:(int64_t)bytesSent totalBytesSent:(int64_t)totalBytesSent totalBytesExpectedToSend:(int64_t)totalBytesExpectedToSend;
+/** The task has received a response and no further messages will be received until the completion block is called. */
+- (void)URLSession:(NSURLSession *)session dataTask:(NSURLSessionDataTask *)dataTask didReceiveResponse:(NSURLResponse *)response completionHandler:(void (^)(NSURLSessionResponseDisposition disposition))completionHandler;
+/** Sent when data is available for the delegate to consume. Data may be discontiguous */
+- (void)URLSession:(NSURLSession *)session dataTask:(NSURLSessionDataTask *)dataTask didReceiveData:(NSData *)data;
+/** Sent as the last message related to a specific task.  A nil Error implies that no error occurred and this task is complete. */
+- (void)URLSession:(NSURLSession *)session task:(NSURLSessionTask *)task didCompleteWithError:(nullable NSError *)error;
+/** Asks the delegate if it needs to store responses in the cache. */
+- (void)URLSession:(NSURLSession *)session dataTask:(NSURLSessionDataTask *)dataTask willCacheResponse:(NSCachedURLResponse *)proposedResponse completionHandler:(void (^)(NSCachedURLResponse *cachedResponse))completionHandler;
+@end
+
 @implementation FAppleHttpResponseDelegate
 @synthesize Response;
 @synthesize RequestStatus;
@@ -433,13 +482,13 @@
  * FAppleHttpRequest implementation
  ***************************************************************************/
 
-FAppleHttpRequest::FAppleHttpRequest(NSURLSession* InSession, NSURLSession* InBackgroundSession)
+FAppleHttpRequest::FAppleHttpRequest(NSURLSession* InSession)
 :   Session([InSession retain])
-,	BackgroundSession([InBackgroundSession retain])
 ,   Task(nil)
 ,	ContentBytesLength(0)
 ,	LastReportedBytesWritten(0)
 ,	LastReportedBytesRead(0)
+,	bStartTaskInHttpThread(CVarAppleHttpStartTaskInHttpThreadEnabled.GetValueOnAnyThread())
 {
 	bUsePlatformActivityTimeout = false;
 	
@@ -663,30 +712,12 @@ bool FAppleHttpRequest::ProcessRequest()
 	
 	StartProcessTime = FPlatformTime::Seconds();
 
-	SetStatus(EHttpRequestStatus::Processing);
-	SetFailureReason(EHttpFailureReason::None);
-	// AppleEventLoop sets a delegate into the response to be able to notify events
-	InitResponse();
-
-	// if we specified the option to use ImmediateRequest and the feature is enabled (which means the BackgroundSession should not be nil)
-	if (GetOption(HttpRequestOptions::RequestMode) == LexToString(EHttpRequestMode::ImmediateRequest) 
-		&& FApplePlatformHttp::IsBackgroundRequestFeatureEnabled() && BackgroundSession != nil)
-	{
-
-		FHttpModule::Get().GetHttpManager().AddRequest(SharedThis(this));
-
-		// for background sessions we must use the download/upload APIs instead of the dataTask ones
-		// it is a requirement for the background session to handle these requests while in background
-		Task = [BackgroundSession uploadTaskWithStreamedRequest:Request];
-
-		UE_LOG(LogHttp, VeryVerbose, TEXT("Starting immediate request %p with task %p"), this, Task);
-
-		FApplePlatformHttp::Add([Task retain], SharedThis(this));
-
-		[Task resume];
-	}
-	else
+	if (bStartTaskInHttpThread)
 	{
+		SetStatus(EHttpRequestStatus::Processing);
+		SetFailureReason(EHttpFailureReason::None);
+		// AppleEventLoop sets a delegate into the response to be able to notify events
+		InitResponse();
 		FHttpModule::Get().GetHttpManager().AddThreadedRequest(SharedThis(this));
 	}
 	return true;
@@ -761,7 +792,40 @@ bool FAppleHttpRequest::SetupRequest()
 		return false;
 	}
 	
-	return true;
+	if (bStartTaskInHttpThread)
+	{
+		return true;
+	}
+	else
+	{
+		Task = [Session dataTaskWithRequest: Request];
+	
+		if (Task != nil)
+		{
+			SetStatus(EHttpRequestStatus::Processing);
+			SetFailureReason(EHttpFailureReason::None);
+
+			InitResponse();
+
+			// Both Task and Response keep a strong reference to the delegate
+			TSharedPtr<FAppleHttpResponse> Response = StaticCastSharedPtr<FAppleHttpResponse>(ResponseCommon);
+			Task.delegate = Response->ResponseDelegate;
+
+			//Setup delegates before starting the request
+			FHttpModule::Get().GetHttpManager().AddThreadedRequest(SharedThis(this));
+
+			[[Task retain] resume];
+			UE_LOG(LogHttp, Verbose, TEXT("Task started %p"), this);
+			return true;
+		}
+		else
+		{
+			UE_LOG(LogHttp, Warning, TEXT("SetupRequest failed. Could not initialize NSURLSessionTask."));
+			SetStatus(EHttpRequestStatus::Failed);
+			SetFailureReason(EHttpFailureReason::ConnectionError);
+			return false;
+		}	
+	}
 }
 
 FHttpResponsePtr FAppleHttpRequest::CreateResponse()
@@ -915,34 +979,36 @@ void FAppleHttpRequest::CheckProgressDelegate()
 
 bool FAppleHttpRequest::StartThreadedRequest()
 {
-	if (bCanceled)
+	if (bStartTaskInHttpThread)
 	{
-		UE_LOG(LogHttp, Verbose, TEXT("StartThreadedRequest ignored because request has been canceled. %s url=%s"), *GetVerb(), *GetURL());
-		return false;
-	}
+		if (bCanceled)
+		{
+			UE_LOG(LogHttp, Verbose, TEXT("StartThreadedRequest ignored because request has been canceled. %s url=%s"), *GetVerb(), *GetURL());
+			return false;
+		}
 
-	if (Task != nil)
-	{
-		UE_LOG(LogHttp, Verbose, TEXT("StartThreadedRequest ignored because task was already in progress. %s url=%s"), *GetVerb(), *GetURL());
-		return false;
-	}
+		if (Task != nil)
+		{
+			UE_LOG(LogHttp, Verbose, TEXT("StartThreadedRequest ignored because task was already in progress. %s url=%s"), *GetVerb(), *GetURL());
+			return false;
+		}
+	
+		Task = [Session dataTaskWithRequest: Request];
+	
+		if (Task == nil)
+		{
+			UE_LOG(LogHttp, Warning, TEXT("StartThreadedRequest failed. Could not initialize NSURLSessionTask."));
+			SetStatus(EHttpRequestStatus::Failed);
+			SetFailureReason(EHttpFailureReason::ConnectionError);
+			return false;
+		}
 
-	Task = [Session dataTaskWithRequest: Request];
+		// Both Task and Response keep a strong reference to the delegate
+		TSharedPtr<FAppleHttpResponse> Response = StaticCastSharedPtr<FAppleHttpResponse>(ResponseCommon);
+		Task.delegate = Response->ResponseDelegate;
 
-	if (Task == nil)
-	{
-		UE_LOG(LogHttp, Warning, TEXT("StartThreadedRequest failed. Could not initialize NSURLSessionTask."));
-		SetStatus(EHttpRequestStatus::Failed);
-		SetFailureReason(EHttpFailureReason::ConnectionError);
-		return false;
+		[[Task retain] resume];
 	}
-
-	// Both Task and Response keep a strong reference to the delegate
-	TSharedPtr<FAppleHttpResponse> Response = StaticCastSharedPtr<FAppleHttpResponse>(ResponseCommon);
-	Task.delegate = Response->ResponseDelegate;
-
-	[[Task retain] resume];
-	
 	return true;
 }
 
@@ -989,11 +1055,6 @@ void FAppleHttpResponse::SetHeaders(TMap<FString, FString>&& InHeaders)
 	Headers = MoveTemp(InHeaders);
 }
 
-FAppleHttpResponseDelegate* FAppleHttpResponse::GetResponseDelegate() const
-{
-	return ResponseDelegate;
-}
-
 void FAppleHttpResponse::CleanSharedObjects()
 {
 	[ResponseDelegate CleanSharedObjects];
diff --git a/Engine/Source/Runtime/Online/HTTP/Private/Apple/AppleHttp.h b/Engine/Source/Runtime/Online/HTTP/Private/Apple/AppleHttp.h
--- a/Engine/Source/Runtime/Online/HTTP/Private/Apple/AppleHttp.h
+++ b/Engine/Source/Runtime/Online/HTTP/Private/Apple/AppleHttp.h
@@ -57,9 +57,8 @@ public:
 	 * Constructor
 	 *
 	 * @param InSession - NSURLSession session used to create NSURLSessionTask to retrieve the response
-	 * @param InBackgroundSession - NSURLSession session used to create NSURLSessionTask to retrieve the response in the background
 	 */
-	FAppleHttpRequest(NSURLSession* InSession, NSURLSession* InBackgroundSession);
+	explicit FAppleHttpRequest(NSURLSession* InSession);
 
 	/**
 	 * Destructor. Clean up any connection/request handles
@@ -109,11 +108,9 @@ private:
 	/** This is the NSMutableURLRequest, all our Apple functionality will deal with this. */
 	NSMutableURLRequest* Request;
 
-	/** This is the session our request belongs to (depending on the RequestMode option) */
+	/** This is the session our request belongs to */
 	NSURLSession* Session;
-	/** This is the background session our request belongs to (depending on the RequestMode option) */
-	NSURLSession* BackgroundSession;
-
+	
 	/** This is the Task associated to the sessionin charge of our request */
 	NSURLSessionTask* Task;
     
@@ -138,6 +135,9 @@ private:
 
 	/** Last reported bytesread */
 	int32 LastReportedBytesRead;
+	
+	// UE-228716: This request should follow same logic as the rest of implementations
+	const bool bStartTaskInHttpThread;
 };
 
 @class FAppleHttpResponseDelegate;
@@ -182,8 +182,6 @@ public:
 
 	void SetHeaders(TMap<FString, FString>&& InHeaders);
 
-	FAppleHttpResponseDelegate* GetResponseDelegate() const;
-
 	using FHttpResponseCommon::AppendToPayload;
 
 	/**
diff --git a/Engine/Source/Runtime/Online/HTTP/Private/Apple/AppleHttpBackgroundSession.cpp b/Engine/Source/Runtime/Online/HTTP/Private/Apple/AppleHttpBackgroundSession.cpp
deleted file mode 100644
diff --git a/Engine/Source/Runtime/Online/HTTP/Private/Apple/AppleHttpBackgroundSession.h b/Engine/Source/Runtime/Online/HTTP/Private/Apple/AppleHttpBackgroundSession.h
deleted file mode 100644
diff --git a/Engine/Source/Runtime/Online/HTTP/Private/Apple/AppleHttpObjc.h b/Engine/Source/Runtime/Online/HTTP/Private/Apple/AppleHttpObjc.h
deleted file mode 100644
diff --git a/Engine/Source/Runtime/Online/HTTP/Private/Apple/ApplePlatformHttp.cpp b/Engine/Source/Runtime/Online/HTTP/Private/Apple/ApplePlatformHttp.cpp
--- a/Engine/Source/Runtime/Online/HTTP/Private/Apple/ApplePlatformHttp.cpp
+++ b/Engine/Source/Runtime/Online/HTTP/Private/Apple/ApplePlatformHttp.cpp
@@ -1,15 +1,12 @@
 // Copyright Epic Games, Inc. All Rights Reserved.
 
 #include "Apple/ApplePlatformHttp.h"
-#include "HAL/IConsoleManager.h"
 #include "Http.h"
 #include "AppleHttpManager.h"
-#include "AppleHttpBackgroundSession.h"
 #include "AppleHttp.h"
 #include "Apple/CFRef.h"
 #include "Misc/Base64.h"
 #include "Misc/CommandLine.h"
-#include "Misc/ScopeRWLock.h"
 
 #include "CommonCrypto/CommonDigest.h"
 
@@ -217,16 +214,6 @@ static bool ValidateCertificatePublicKeyPinningApple(NSURLAuthenticationChalleng
 
 #endif
 
-static TAutoConsoleVariable<bool> CVarEnableBackgroundRequests(
-	TEXT("AppleHttp.EnableBackgroundRequests"),
-	true,
-	TEXT("If true, will use the Background URLSession to schedule tasks that can run while the app is in the background. "
-			   "If false, will use the current flow scheduling the request to be executed."),
-	ECVF_Default);
-
-FRWLock FApplePlatformHttp::TaskRequestMapLock;
-TMap<NSURLSessionTask*, TSharedPtr<FAppleHttpRequest>> FApplePlatformHttp::TaskRequestMap;
-
 void FApplePlatformHttp::Init()
 {
 #if WITH_SSL
@@ -264,23 +251,6 @@ void FApplePlatformHttp::InitWithNSUrlSession()
 	
     [Session retain];
 
-#if PLATFORM_IOS // won't work on Mac, tvOS and visionOS
-	if (IsBackgroundRequestFeatureEnabled()) // can be disabled by cvar
-	{
-		Config = [NSURLSessionConfiguration backgroundSessionConfigurationWithIdentifier:@"com.unrealengine.backgroundhttpsession"];
-		
-		// Disable cache to mimic WinInet behavior
-		Config.requestCachePolicy = NSURLRequestReloadIgnoringLocalCacheData;
-		// Set discretionary to false so the OS will send the request ASAP (with fewer considerations on battery level, power consumption, and such)
-		Config.discretionary = NO;
-		Config.sessionSendsLaunchEvents = YES;
-
-		FApplePlatformHttpBackgroundSessionDelegate* Delegate = [[FApplePlatformHttpBackgroundSessionDelegate alloc] init];
-		BackgroundSession = [NSURLSession sessionWithConfiguration:Config delegate:Delegate delegateQueue:nil];
-
-		[BackgroundSession retain];
-	}
-#endif // PLATFORM_IOS
 }
 
 void FApplePlatformHttp::ShutdownWithNSUrlSession()
@@ -288,10 +258,6 @@ void FApplePlatformHttp::ShutdownWithNSUrlSession()
 	[Session invalidateAndCancel];
 	[Session release];
 	Session = nil;
-
-	[BackgroundSession invalidateAndCancel];
-	[BackgroundSession release];
-	BackgroundSession = nil;
 }
 
 FHttpManager* FApplePlatformHttp::CreatePlatformHttpManager()
@@ -306,41 +272,6 @@ IHttpRequest* FApplePlatformHttp::ConstructRequest()
 		InitWithNSUrlSession();
 	}
 
-	// To avoid breaking the ConstructRequest contract I'll be passing both sessions here
-	// If the background feature is disabled, BackgroundSession will be nil and I can check that inside the FAppleHttpRequest
-	FAppleHttpRequest* Request = new FAppleHttpRequest(Session, BackgroundSession);
-
-	return Request;
-}
-
-bool FApplePlatformHttp::IsBackgroundRequestFeatureEnabled()
-{
-	return CVarEnableBackgroundRequests.GetValueOnAnyThread();
-}
-
-TSharedPtr<FAppleHttpRequest>& FApplePlatformHttp::Add(NSURLSessionTask* Task, const TSharedPtr<FAppleHttpRequest>& Request)
-{
-	FRWScopeLock ScopeLock(TaskRequestMapLock, FRWScopeLockType::SLT_Write);
-
-	return TaskRequestMap.Add(Task, CopyTemp(Request));
+	return new FAppleHttpRequest(Session);
 }
 
-bool FApplePlatformHttp::Find(NSURLSessionTask* Task, TSharedPtr<FAppleHttpRequest>& Request)
-{
-	FRWScopeLock ScopeLock(TaskRequestMapLock, FRWScopeLockType::SLT_ReadOnly);
-
-	if (TSharedPtr<FAppleHttpRequest>* Found = TaskRequestMap.Find(Task))
-	{
-		Request = *Found;
-		return true;
-	}
-	
-	return false;
-}
-
-bool FApplePlatformHttp::RemoveAndCopyValue(NSURLSessionTask* Task, TSharedPtr<FAppleHttpRequest>& Request)
-{
-	FRWScopeLock ScopeLock(TaskRequestMapLock, FRWScopeLockType::SLT_Write);
-
-	return TaskRequestMap.RemoveAndCopyValue(Task, Request);
-}
\ No newline at end of file
diff --git a/Engine/Source/Runtime/Online/HTTP/Private/Interfaces/IHttpRequest.cpp b/Engine/Source/Runtime/Online/HTTP/Private/Interfaces/IHttpRequest.cpp
--- a/Engine/Source/Runtime/Online/HTTP/Private/Interfaces/IHttpRequest.cpp
+++ b/Engine/Source/Runtime/Online/HTTP/Private/Interfaces/IHttpRequest.cpp
@@ -2,19 +2,6 @@
 
 #include "Interfaces/IHttpRequest.h"
 
-const TCHAR* LexToString(EHttpRequestMode Mode)
-{
-	switch (Mode)
-	{
-		case EHttpRequestMode::ThreadedRequest: return TEXT("ThreadedRequest");
-		case EHttpRequestMode::ImmediateRequest: return TEXT("ImmediateRequest");
-
-		default:
-		checkf(false, TEXT("Please, implement new modes here"));
-		return TEXT("Unknown");
-	}
-}
-
 class FArchiveWithDelegateV2 final : public FArchive
 {
 public:
diff --git a/Engine/Source/Runtime/Online/HTTP/Public/Apple/ApplePlatformHttp.h b/Engine/Source/Runtime/Online/HTTP/Public/Apple/ApplePlatformHttp.h
--- a/Engine/Source/Runtime/Online/HTTP/Public/Apple/ApplePlatformHttp.h
+++ b/Engine/Source/Runtime/Online/HTTP/Public/Apple/ApplePlatformHttp.h
@@ -1,7 +1,6 @@
 // Copyright Epic Games, Inc. All Rights Reserved.
 
 #pragma once
-
 #include "GenericPlatform/GenericPlatformHttp.h"
 
 /**
@@ -10,6 +9,7 @@
 class HTTP_API FApplePlatformHttp : public FGenericPlatformHttp
 {
 public:
+
 	/**
 	 * Platform initialization step
 	 */
@@ -34,19 +34,9 @@ public:
 	 */
 	static IHttpRequest* ConstructRequest();
 
-	static bool IsBackgroundRequestFeatureEnabled();
-
-	static TSharedPtr<class FAppleHttpRequest>& Add(NSURLSessionTask* Task, const TSharedPtr<class FAppleHttpRequest>& Request);
-	static bool Find(NSURLSessionTask* Task, OUT TSharedPtr<class FAppleHttpRequest>& Request);
-	static bool RemoveAndCopyValue(NSURLSessionTask* Task, OUT TSharedPtr<class FAppleHttpRequest>& Request);
-
 private:
     /** Session used to create Apple based requests */
     static inline NSURLSession* Session = nil;
-	static inline NSURLSession* BackgroundSession = nil;
-
-	static FRWLock TaskRequestMapLock;
-	static TMap<NSURLSessionTask*, TSharedPtr<class FAppleHttpRequest>> TaskRequestMap;
 
 	static void InitWithNSUrlSession();
 	static void ShutdownWithNSUrlSession();
diff --git a/Engine/Source/Runtime/Online/HTTP/Public/HttpManager.h b/Engine/Source/Runtime/Online/HTTP/Public/HttpManager.h
--- a/Engine/Source/Runtime/Online/HTTP/Public/HttpManager.h
+++ b/Engine/Source/Runtime/Online/HTTP/Public/HttpManager.h
@@ -321,6 +321,8 @@ protected:
 	HTTP_API void UpdateUrlPatternsToDisableFailedLog(IConsoleVariable* CVar);
 	HTTP_API void UpdateUrlPatternsToMockResponse(IConsoleVariable* CVar);
 
+	HTTP_API void AddRequest(const FHttpRequestRef& Request);
+
 protected:
 	/** List of Http requests that are actively being processed */
 	TArray<FHttpRequestRef> Requests;
@@ -449,13 +451,6 @@ PACKAGE_SCOPE:
 	HTTP_API void MarkCurrentThreadCompletingRequest(bool bCompleting);
 	HTTP_API bool IsCurrentThreadCompletingRequest() const;
 
-	/**
-	 * Add a http request to the manager without enqueueing the threaded execution
-	 * 
-	 * @param Request - the request object to add
-	 */
-	HTTP_API void AddRequest(const FHttpRequestRef& Request);
-
 protected:
 	TMap<FString, FMockResponse> UrlPatternsToMockResponse;
 	mutable FCriticalSection UrlPatternsToMockResponseCriticalSection;
diff --git a/Engine/Source/Runtime/Online/HTTP/Public/Interfaces/IHttpRequest.h b/Engine/Source/Runtime/Online/HTTP/Public/Interfaces/IHttpRequest.h
--- a/Engine/Source/Runtime/Online/HTTP/Public/Interfaces/IHttpRequest.h
+++ b/Engine/Source/Runtime/Online/HTTP/Public/Interfaces/IHttpRequest.h
@@ -131,25 +131,11 @@ private:
 namespace HttpRequestOptions
 {
 	static const FName HttpVersion("HttpVersion");
-	static const FName RequestMode("RequestMode");
 #if UE_HTTP_SUPPORT_UNIX_SOCKET
 	static const FName UnixSocketPath("UnixSocketPath");
 #endif //UE_HTTP_SUPPORT_UNIX_SOCKET
 }
 
-/**
- * Mode that can be specified through IHttpRequest::SetOption(HttpRequestOptions::RequestMode, LexToString(EHttpRequestMode::Value));
- */
-enum class EHttpRequestMode : uint8
-{
-	/** The request will be enqueued to the http manager thread system */
-	ThreadedRequest,
-	/** The request will be immediately sent when the platform supports it */
-	ImmediateRequest
-};
-
-HTTP_API const TCHAR* LexToString(EHttpRequestMode Mode);
-
 /**
  * Interface for Http requests (created using FHttpFactory)
  */
diff --git a/Engine/Source/Runtime/StudioTelemetry/Private/AnalyticsProviderBroadcast.cpp b/Engine/Source/Runtime/StudioTelemetry/Private/AnalyticsProviderBroadcast.cpp
--- a/Engine/Source/Runtime/StudioTelemetry/Private/AnalyticsProviderBroadcast.cpp
+++ b/Engine/Source/Runtime/StudioTelemetry/Private/AnalyticsProviderBroadcast.cpp
@@ -2,7 +2,6 @@
 
 #include "AnalyticsProviderBroadcast.h"
 #include "Analytics.h"
-#include "Interfaces/IAnalyticsProvider.h"
 #include "Interfaces/IAnalyticsProviderModule.h"
 #include "AnalyticsProviderConfigurationDelegate.h"
 #include "Misc/ConfigCacheIni.h"
@@ -283,14 +282,14 @@ static void CheckForDuplicateAttributes(const FString& EventName, const TArray<F
 #endif
 }
 
-void FAnalyticsProviderBroadcast::RecordEvent(const FString& EventName, const TArray<FAnalyticsEventAttribute>& Attributes, EAnalyticsRecordEventMode Mode)
+void FAnalyticsProviderBroadcast::RecordEvent(const FString& EventName, const TArray<FAnalyticsEventAttribute>& Attributes)
 {
 	// Check for duplicate attributes against themselves
 	CheckForDuplicateAttributes(EventName, Attributes);
 
 	for (TProviders::TConstIterator it(Providers); it; ++it)
 	{
-		(*it).Value->RecordEvent(EventName, Attributes, Mode);
+		(*it).Value->RecordEvent(EventName, Attributes);
 	}
 
 	if (RecordEventCallback)
diff --git a/Engine/Source/Runtime/StudioTelemetry/Private/AnalyticsProviderBroadcast.h b/Engine/Source/Runtime/StudioTelemetry/Private/AnalyticsProviderBroadcast.h
--- a/Engine/Source/Runtime/StudioTelemetry/Private/AnalyticsProviderBroadcast.h
+++ b/Engine/Source/Runtime/StudioTelemetry/Private/AnalyticsProviderBroadcast.h
@@ -49,12 +49,7 @@ public:
 	virtual FString GetSessionID() const override;
 	virtual FString GetUserID() const override;
 	
-	virtual void RecordEvent(const FString& EventName, const TArray<FAnalyticsEventAttribute>& Attributes) override 
-	{
-		RecordEvent(EventName, Attributes, EAnalyticsRecordEventMode::Cached);
-	}
-	
-	virtual void RecordEvent(const FString& EventName, const TArray<FAnalyticsEventAttribute>& Attributes, EAnalyticsRecordEventMode Mode) override;	
+	virtual void RecordEvent(const FString& EventName, const TArray<FAnalyticsEventAttribute>& Attributes) override;	
 	void SetRecordEventCallback(OnRecordEvent Callback);
 	bool HasValidProviders() const { return Providers.Num() > 0; }
 	TWeakPtr<IAnalyticsProvider> GetAnalyticsProvider(const FString& Name);
