MakeAnalyticsEventAttributeArray refactor for code size

diff --git a/Engine/Source/Runtime/Analytics/Analytics/Public/AnalyticsEventAttribute.h b/Engine/Source/Runtime/Analytics/Analytics/Public/AnalyticsEventAttribute.h
--- a/Engine/Source/Runtime/Analytics/Analytics/Public/AnalyticsEventAttribute.h
+++ b/Engine/Source/Runtime/Analytics/Analytics/Public/AnalyticsEventAttribute.h
@@ -248,44 +248,30 @@ PRAGMA_ENABLE_DEPRECATION_WARNINGS
 /** Helper functions for MakeAnalyticsEventAttributeArray. */
 namespace ImplMakeAnalyticsEventAttributeArray
 {
-	template <typename Allocator>
-	struct Helper final
+	// AddElement actually adds an element to the attributes array and is flagged no-inline so the array insert/value conversion
+	// code is generated once per type combination.
+	template <typename Allocator, typename KeyType, typename ValueType>
+	FORCENOINLINE void AddElement(TArray<FAnalyticsEventAttribute, Allocator>& Attrs, const KeyType& Key, const ValueType& Value)
 	{
-		Helper(TArray<FAnalyticsEventAttribute, Allocator>& Attrs) : Attrs(Attrs) {}
-
-		TOptional<FString> LastKey;
-		TArray<FAnalyticsEventAttribute, Allocator>& Attrs;
-
-		template <typename Type>
-		void ForEach(Type&& Thing)
-		{
-			if (!LastKey)
-			{
-				// Even indices are keys.
-				if constexpr (std::is_convertible_v<Type, FString>)
-				{
-					LastKey = ::MoveTempIfPossible(Thing);
-				}
-				else
-				{
-					checkf(false, TEXT("Even indices are keys and must be convertible to `FString`!"));
-				}
-			}
-			else
-			{
-				// Odd indices are values.
-				Attrs.Emplace(::MoveTemp(*LastKey), Thing);
-				LastKey.Reset();
-			}
-		}
-	};
+		static_assert(std::is_convertible_v<KeyType, FString>, "Keys must be convertible to `FString`!");
+		Attrs.Emplace(FString{Key}, Value);
+	}
 
-	template <typename Allocator, typename...ArgTypes>
-	inline void MakeArray(TArray<FAnalyticsEventAttribute, Allocator>& Attrs, ArgTypes&&...Args)
+	template <typename Allocator>
+    FORCEINLINE void MakeArray(TArray<FAnalyticsEventAttribute, Allocator>& Attrs)
+    {
+    }
+
+	// MakeArray is just a helper whose only purpose is to generate a sequence of AddElement calls and is meant to be always inlined.
+    template <typename Allocator, typename KeyType, typename ValueType, typename... RemainingArgTypes>
+    FORCEINLINE void MakeArray(TArray<FAnalyticsEventAttribute, Allocator>& Attrs, KeyType&& Key, ValueType&& Value, RemainingArgTypes&& ...RemainingArgs)
 	{
-		static_assert(sizeof...(Args) % 2 == 0, "Must pass an even number of arguments.");
-		Helper<Allocator> H(Attrs);
-		(H.ForEach(Args), ...);
+		static_assert(sizeof...(RemainingArgs) % 2 == 0, "Must pass an even number of arguments.");
+
+		// Decay here to decay string literals into const char*/const TCHAR* (losing the information about their length)
+		// to limit the number of unique instantiations of AddElement.
+		AddElement<Allocator, std::decay_t<KeyType>, std::decay_t<ValueType>>(Attrs, Forward<KeyType>(Key), Forward<ValueType>(Value));
+		MakeArray(Attrs, Forward<RemainingArgTypes>(RemainingArgs)...);
 	}
 }
 
