Add Analytics Delegate FOnAnalyticsEventQueued to IAnalyticsProviderET to process events being added to the analytics queue

#rnx

diff --git a/Engine/Source/Runtime/Analytics/AnalyticsET/Private/IAnalyticsProviderET.cpp b/Engine/Source/Runtime/Analytics/AnalyticsET/Private/IAnalyticsProviderET.cpp
--- a/Engine/Source/Runtime/Analytics/AnalyticsET/Private/IAnalyticsProviderET.cpp
+++ b/Engine/Source/Runtime/Analytics/AnalyticsET/Private/IAnalyticsProviderET.cpp
@@ -157,6 +157,7 @@ public:
 	virtual void BlockUntilFlushed(float InTimeoutSec) override;
 	virtual void SetShouldRecordEventFunc(const ShouldRecordEventFunction& InShouldRecordEventFunc) override;
 	virtual FOnPreAnalyticsEventProcessed& OnPreAnalyticsEventProcessed() override { return OnPreAnalyticsEventProcessedDelegate; }
+	virtual FOnAnalyticsEventQueued& OnAnalyticsEventQueued() override { return OnAnalyticsEventQueuedDelegate; }
 
 	virtual ~FAnalyticsProviderET();
 
@@ -215,6 +216,7 @@ private:
 	FAnalyticsUserAgentCache UserAgentCache;
 
 	FOnPreAnalyticsEventProcessed OnPreAnalyticsEventProcessedDelegate;
+	FOnAnalyticsEventQueued OnAnalyticsEventQueuedDelegate;
 };
 
 TSharedPtr<IAnalyticsProviderET> FAnalyticsET::CreateAnalyticsProvider(const Config& ConfigValues) const
@@ -660,43 +662,60 @@ void FAnalyticsProviderET::RecordEvent(FString&& EventName, const TArray<FAnalyt
 	// let higher level code filter the decision of whether to send the event
 	if (ShouldRecordEvent(EventName))
 	{
-		switch (Mode)
+		bool bQueueEvent = true;
+		const FAnalyticEventQueuedInfo AnalyticEventInfo =
 		{
-			case EAnalyticsRecordEventMode::Cached:
+			Mode,
+			EventName,
+			Attributes
+		};
+
+		OnAnalyticsEventQueuedDelegate.Broadcast(bQueueEvent, AnalyticEventInfo);
+
+		if (bQueueEvent)
+		{
+			switch (Mode)
 			{
-				// fire any callbacks
-				for (const auto& Cb : EventRecordedCallbacks)
+				case EAnalyticsRecordEventMode::Cached:
 				{
-					// we no longer track if the event was Json, each attribute does.
-					Cb(EventName, Attributes, false);
-				}
+					// fire any callbacks
+					for (const auto& Cb : EventRecordedCallbacks)
+					{
+						// we no longer track if the event was Json, each attribute does.
+						Cb(EventName, Attributes, false);
+					}
 
-				if (!Config.UseLegacyProtocol)
-				{
-					EventCache.AddToCache(MoveTemp(EventName), Attributes);
+					if (!Config.UseLegacyProtocol)
+					{
+						EventCache.AddToCache(MoveTemp(EventName), Attributes);
 
-					// if we aren't caching events, flush immediately. This is really only for debugging as it will significantly affect bandwidth.
-					if (!bShouldCacheEvents)
+						// if we aren't caching events, flush immediately. This is really only for debugging as it will significantly affect bandwidth.
+						if (!bShouldCacheEvents)
+						{
+							FlushEvents();
+						}
+					}
+					else
 					{
-						FlushEvents();
+						FlushEventLegacy(EventName, Attributes);
 					}
 				}
-				else
+				break;
+
+				case EAnalyticsRecordEventMode::Immediate:
 				{
-					FlushEventLegacy(EventName, Attributes);
+					SendImmediately(MoveTemp(EventName), Attributes);
 				}
-			}
-			break;
+				break;
 
-			case EAnalyticsRecordEventMode::Immediate:
-			{
-				SendImmediately(MoveTemp(EventName), Attributes);
+				default:
+					checkf(false, TEXT("Please, implement new modes here"));
+					break;
 			}
-			break;
-
-			default:
-			checkf(false, TEXT("Please, implement new modes here"));
-			break;
+		}
+		else
+		{
+			UE_LOG(LogAnalytics, Verbose, TEXT("Dropping Event %s due to OnAnalyticsEventQueuedDelegate request."), *EventName);
 		}
 	}
 	else
diff --git a/Engine/Source/Runtime/Analytics/AnalyticsET/Public/IAnalyticsProviderET.h b/Engine/Source/Runtime/Analytics/AnalyticsET/Public/IAnalyticsProviderET.h
--- a/Engine/Source/Runtime/Analytics/AnalyticsET/Public/IAnalyticsProviderET.h
+++ b/Engine/Source/Runtime/Analytics/AnalyticsET/Public/IAnalyticsProviderET.h
@@ -19,6 +19,21 @@ class IHttpRequest;
 /** Fired immediately prior to the HTTP request for an analytics event being processed by HTTP */
 DECLARE_MULTICAST_DELEGATE_TwoParams(FOnPreAnalyticsEventProcessed, TSharedRef<IHttpRequest, ESPMode::ThreadSafe>& /*HttpRequest*/);
 
+struct FAnalyticEventQueuedInfo
+{
+	const EAnalyticsRecordEventMode Mode;
+	const FString& EventName;
+	const TArray<FAnalyticsEventAttribute>& Attributes;
+};
+
+/** 
+*	Fired on every analytic event being queued
+*   boolean value should only be left alone or set to false, which will cause the event to not be queued.
+*   FAnalyticEventQueuedInfo is only valid during the broadcast, you must copy any data from it that you need to keep
+*/
+DECLARE_MULTICAST_DELEGATE_TwoParams(FOnAnalyticsEventQueued, bool& /* bQueueEvent If this boolean is made false by any registered callback, this event will not be queued.  Do not set this value to true */, 
+	const FAnalyticEventQueuedInfo& /* All the details on the event that is being queued */)
+
 /** ET specific analytics provider instance. Exposes additional APIs to support Json-based events, default attributes, and allowing events to be disabled (generally via hotfixing). */
 class IAnalyticsProviderET : public IAnalyticsProvider
 {
@@ -151,6 +166,7 @@ public:
 	virtual void SetShouldRecordEventFunc(const ShouldRecordEventFunction& ShouldRecordEventFunc) = 0;
 	
 	virtual FOnPreAnalyticsEventProcessed& OnPreAnalyticsEventProcessed() = 0;
+	virtual FOnAnalyticsEventQueued& OnAnalyticsEventQueued() = 0;
 
 private:
 	/** Needed to support the old, unsafe GetDefaultAttributes() API. */
