Studio Telemetry Updates

Added Test Cases CSV files.
Added a AnalyticsProviderCSV which will write out events to one CSV file per event.
CSV output per event will be usefule to validate autimated tests of CSVToTelemetry.
Sanitized the JSON output from AnalyticsLog to ensure JSON output is well formatted.
Moved the AnalyticsLog, AnalyticsCSV, AnalyticsHorde provider examples to the public facing BaseEngine.ini. These are disabled by default.

#Synced-CL 48759921

diff --git a/Engine/Config/BaseEngine.ini b/Engine/Config/BaseEngine.ini
--- a/Engine/Config/BaseEngine.ini
+++ b/Engine/Config/BaseEngine.ini
@@ -3808,23 +3808,39 @@ SendUserData=false
 SendHardwareData=true
 SendOSData=true
 
-; Example Log File Analytics
-;[StudioTelemetry.Provider.Log] 
-; Name=LogAnalytics
-; ProviderModule=AnalyticsLog
-; UsageType=Editor
-; Override the default output location from Saved/Telemetry/Telemetry.json
-; FileName="somefile.json"
-; FolderPath="somefolder"
-
-; Example Horde Analytics Settings
-;[StudioTelemetry.Provider.Horde]
-; Name=HordeAnalytics
-; ProviderModule=AnalyticsET
-; UsageType=EditorAndClient
-; APIKeyET=HordeAnalytics.Dev
-; APIServerET="https://myownhordeserver.com/" ; Point Server URL to your own Horde server
-; APIEndpointET="api/v1/telemetry/engine"
+; Example Analytics Provider. 
+; Writes out Telemetry Events to a single JSON log file. 
+; Disabled by default
+[StudioTelemetry.Provider.Log] 
+Name=LogAnalytics
+Enabled=False
+ProviderModule=AnalyticsLog
+UsageType=EditorAndClient
+FileName=""
+FolderPath=""
+
+; Example Analytics Provider. 
+; Writes out Telemetry Events to a multiple CSV files, one file for each event instance.
+; Disabled by default
+[StudioTelemetry.Provider.CSV] 
+Name=CSVAnalytics
+Enabled=False
+ProviderModule=AnalyticsCSV
+UsageType=EditorAndClient
+FolderPath=""
+
+; Example Horde Analytics Provider
+; Sends Telemetry Events to Horde Analytics
+; Disabled by default
+; Point the APIServerET to your own Horde Server instance
+[StudioTelemetry.Provider.Horde]
+Name=HordeAnalyticsProd
+ProviderModule=AnalyticsHorde
+Enabled=False
+UsageType=EditorAndClient
+APIKeyET=HordeAnalytics.Prod
+APIServerET="https://myownhordeserver.com/"
+APIEndpointET="api/v1/telemetry/engine"
 
 ; Example Horde UBA Distribution Settings - See Engine/Source/Programs/UnrealBuildTool/Executors/UnrealBuildAccelerator/UnrealBuildAcceleratorHordeConfig.cs
 ;[UbaController]
diff --git a/Engine/Source/Programs/CSVTools/CSVToTelemetry/Private/CSVToTelemetry.cpp b/Engine/Source/Programs/CSVTools/CSVToTelemetry/Private/CSVToTelemetry.cpp
--- a/Engine/Source/Programs/CSVTools/CSVToTelemetry/Private/CSVToTelemetry.cpp
+++ b/Engine/Source/Programs/CSVTools/CSVToTelemetry/Private/CSVToTelemetry.cpp
@@ -200,7 +200,7 @@ static int GenerateTelemetryFromCSVFile(const FString& FilePath)
 					}
 
 					// Keys and values count should always be the same
-					if (KeyArray.Num() >= ValueArray.Num())
+					if (KeyArray.Num() == ValueArray.Num())
 					{
 						ValueSize = Row.NumBytes();
 
@@ -234,8 +234,8 @@ static int GenerateTelemetryFromCSVFile(const FString& FilePath)
 					}
 					else
 					{
-						// The key/value counts don't match
-						UE_LOG(LogCSVToTelemetry, Warning, TEXT("Row %d contains the incorrect value count of %d ( expected %d ) and will be skipped."), RowIndex, ValueArray.Num(), KeyArray.Num());
+						// The key/value counts don't match	
+						UE_LOG(LogCSVToTelemetry, Warning, TEXT("Row %d contains an incorrect number of values %d ( expected %d ) and will be skipped."), RowIndex, ValueArray.Num(), KeyArray.Num());
 					}
 				}
 
@@ -338,7 +338,7 @@ static int GenerateTelemetryFromCSVProfileFile(const FString& FilePath)
 
 static int ShowHelp()
 {
-	UE_LOG(LogCSVToTelemetry, Display, TEXT("\n\nCSVToTelemetry Help\n\nUsage:\n\tCSVToTelemetry.exe -csv=[filename] -event=[eventname] ( -schema=[value] -columns=[name1|name2|....] -attributes=[key1=value1|key2=value2|....] -normalize-bool-strings -maxrows=[value])\n\tCSVToTelemetry.exe -csvprofile=[filename] -event=[eventname]\n\tCSVToTelemetry.exe -help\n\nRequired:\n\t-csv=[filename]\t\t\tGeneric text based csv input file.\n\t-csvprofile=[filename]\t\tCSVProfiler csv input file. Denote binary with .csv.bin otherwise assumed text.\n\t-event=[name]\t\t\tName of telemetry event to send each row to.\nOptional:\n\t-schema=[value]\t\t\tEvent schema value.\n\t-columns=[name1|name2|....]\tColumn include filter.\n\t-normalize-bool-strings\t\tNormalizes any (\"True\"|\"False\"|\"TRUE\"|\"FALSE\") strings to lowercase."));
+	UE_LOG(LogCSVToTelemetry, Display, TEXT("\n\nCSVToTelemetry Help\n\nUsage:\n\tCSVToTelemetry.exe -csv=[filename] -event=[eventname] ( -schema=[value] -columns=[name1|name2|....] -attributes=[key1=value1|key2=value2|....] -normalize-bool-strings -maxrows=[value] )\n\tCSVToTelemetry.exe -csvprofile=[filename] -event=[eventname]\n\tCSVToTelemetry.exe -help\n\nRequired:\n\t-csv=[filename]\t\t\tGeneric text based csv input file.\n\t-csvprofile=[filename]\t\tCSVProfiler csv input file. Denote binary with .csv.bin otherwise assumed text.\n\t-event=[name]\t\t\tName of telemetry event to send each row to.\nOptional:\n\t-schema=[value]\t\t\tEvent schema value.\n\t-columns=[name1|name2|....]\tSpecify column names.\n\t-normalize-bool-strings\t\tNormalizes any (\"True\"|\"False\"|\"TRUE\"|\"FALSE\") strings to lowercase.\n\t-maxrows=[value]\t\tLimits the number of rows sent."));
 	return 0;
 }
 
@@ -358,7 +358,7 @@ INT32_MAIN_INT32_ARGC_TCHAR_ARGV()
 
 	FString FilePath;
 
-	if (FParse::Value(FCommandLine::Get(), TEXT("help"), FilePath) == true)
+	if (FParse::Param(FCommandLine::Get(), TEXT("help")) == true)
 	{
 		// Show help message
 		Result = ShowHelp();
diff --git a/Engine/Source/Runtime/Analytics/AnalyticsCSV/AnalyticsCSV.Build.cs b/Engine/Source/Runtime/Analytics/AnalyticsCSV/AnalyticsCSV.Build.cs
new file mode 100644
--- /dev/null
+++ b/Engine/Source/Runtime/Analytics/AnalyticsCSV/AnalyticsCSV.Build.cs
@@ -0,0 +1,22 @@
+// Copyright Epic Games, Inc. All Rights Reserved.
+
+namespace UnrealBuildTool.Rules
+{
+	public class AnalyticsCSV : ModuleRules
+	{
+		public AnalyticsCSV(ReadOnlyTargetRules Target) : base(Target)
+		{
+			PublicDependencyModuleNames.AddRange(
+				new string[] {
+						"Core"
+				}
+			);
+
+			PrivateDependencyModuleNames.AddRange(
+				new string[] {
+				"Analytics"
+				}
+			);
+		}
+	}
+}
diff --git a/Engine/Source/Runtime/Analytics/AnalyticsCSV/Private/AnalyticsCSV.cpp b/Engine/Source/Runtime/Analytics/AnalyticsCSV/Private/AnalyticsCSV.cpp
new file mode 100644
--- /dev/null
+++ b/Engine/Source/Runtime/Analytics/AnalyticsCSV/Private/AnalyticsCSV.cpp
@@ -0,0 +1,25 @@
+// Copyright Epic Games, Inc. All Rights Reserved.
+
+#include "AnalyticsCSV.h"
+#include "AnalyticsProviderCSV.h"
+
+IMPLEMENT_MODULE( FAnalyticsCSV, AnalyticsCSV );
+
+void FAnalyticsCSV::StartupModule()
+{
+}
+
+void FAnalyticsCSV::ShutdownModule()
+{
+}
+
+TSharedPtr<IAnalyticsProvider> FAnalyticsCSV::CreateAnalyticsProvider(const FAnalyticsProviderConfigurationDelegate& GetConfigValue) const
+{
+	if (GetConfigValue.IsBound())
+	{
+		return MakeShared<FAnalyticsProviderCSV>(GetConfigValue);
+	}
+
+	return nullptr;
+}
+
diff --git a/Engine/Source/Runtime/Analytics/AnalyticsCSV/Private/AnalyticsProviderCSV.cpp b/Engine/Source/Runtime/Analytics/AnalyticsCSV/Private/AnalyticsProviderCSV.cpp
new file mode 100644
--- /dev/null
+++ b/Engine/Source/Runtime/Analytics/AnalyticsCSV/Private/AnalyticsProviderCSV.cpp
@@ -0,0 +1,244 @@
+// Copyright Epic Games, Inc. All Rights Reserved.
+
+#include "AnalyticsProviderCSV.h"
+
+#include "Analytics.h"
+#include "HAL/FileManager.h"
+#include "Misc/Paths.h"
+#include "Tasks/Pipe.h"
+
+FAnalyticsProviderCSV::FAnalyticsProviderCSV(const FAnalyticsProviderConfigurationDelegate& GetConfigValue)
+{
+	FolderPath = GetConfigValue.Execute(TEXT("FolderPath"), true);
+
+	if (FolderPath.IsEmpty())
+	{
+		// See if there's a folder specified in the environment 
+		FolderPath = FPlatformMisc::GetEnvironmentVariable(TEXT("UE_TELEMETRY_DIR"));
+
+		if (FolderPath.IsEmpty())
+		{
+			// Use default output path
+			FolderPath = FPaths::ProjectSavedDir() / TEXT("Telemetry");
+		}
+	}
+
+	if (FolderPath.IsEmpty()==false)
+	{
+		if (IFileManager::Get().MakeDirectory(*FolderPath, true))
+		{
+			WriterPipe = MakeUnique<UE::Tasks::FPipe>(TEXT("FAnalyticsProviderCSV_Writer"));
+		}
+	}
+}
+
+FAnalyticsProviderCSV::~FAnalyticsProviderCSV()
+{
+	// Make sure we close the async write pipe
+	if (WriterPipe.IsValid())
+	{
+		WriterPipe->WaitUntilEmpty();
+		WriterPipe.Reset();
+	}
+
+	// Make sure we close all the CSV file archives
+	for (TArchives::TConstIterator it(Archives); it; ++it)
+	{
+		TSharedPtr<FArchive> Archive = (*it).Value;
+
+		if (Archive.IsValid())
+		{
+			Archive->Flush();
+			Archive->Close();
+			Archive.Reset();
+		}
+	}
+
+	Archives.Reset();
+}
+
+bool FAnalyticsProviderCSV::SetSessionID(const FString& InSessionID)
+{
+	SessionID = InSessionID;
+	return true;
+}
+
+FString FAnalyticsProviderCSV::GetSessionID() const
+{
+	return SessionID;
+}
+
+void FAnalyticsProviderCSV::SetUserID(const FString& InUserID)
+{
+	UserID = InUserID;
+
+}
+
+FString FAnalyticsProviderCSV::GetUserID() const
+{
+	return UserID;
+}
+
+void FAnalyticsProviderCSV::FlushEvents()
+{
+}
+
+void FAnalyticsProviderCSV::SetDefaultEventAttributes(TArray<FAnalyticsEventAttribute>&& Attributes)
+{
+	DefaultEventAttributes = Attributes;
+}
+
+TArray<FAnalyticsEventAttribute> FAnalyticsProviderCSV::GetDefaultEventAttributesSafe() const
+{
+	return DefaultEventAttributes;
+}
+
+int32 FAnalyticsProviderCSV::GetDefaultEventAttributeCount() const
+{
+	return DefaultEventAttributes.Num();
+}
+
+FAnalyticsEventAttribute FAnalyticsProviderCSV::GetDefaultEventAttribute(int AttributeIndex) const
+{
+	return ( AttributeIndex>=0 && AttributeIndex<DefaultEventAttributes.Num() )? DefaultEventAttributes[AttributeIndex]: FAnalyticsEventAttribute();
+}
+
+bool FAnalyticsProviderCSV::StartSession(const TArray<FAnalyticsEventAttribute>& Attributes)
+{
+	TRACE_CPUPROFILER_EVENT_SCOPE(FAnalyticsProviderCSV::StartSession);
+
+	RecordEvent(TEXT("StartSession"), Attributes);
+
+	return true;
+}
+
+void FAnalyticsProviderCSV::EndSession()
+{
+	TRACE_CPUPROFILER_EVENT_SCOPE(FAnalyticsProviderCSV::EndSession);
+
+	RecordEvent(TEXT("EndSession"));
+
+	// Wait until the pipe is empty before removing the archives
+	if (WriterPipe.IsValid())
+	{
+		WriterPipe->WaitUntilEmpty();
+	}
+
+	// Make sure we close all the CSV file archives
+	for (TArchives::TConstIterator it(Archives); it; ++it)
+	{
+		TSharedPtr<FArchive> Archive = (*it).Value;
+
+		if (Archive.IsValid())
+		{
+			Archive->Flush();
+			Archive->Close();
+			Archive.Reset();
+		}
+	}
+
+	Archives.Reset();
+}
+
+void FAnalyticsProviderCSV::RecordEvent(const FString& EventName, const TArray<FAnalyticsEventAttribute>& Attributes)
+{	
+	if (WriterPipe.IsValid() && FolderPath.IsEmpty()==false && EventName.IsEmpty()==false )
+	{
+		// Accumulate all the attributes together. We could have had two loops but this seems cleaner
+		TArray<FAnalyticsEventAttribute> EventAttributes(Attributes);
+
+		if (DefaultEventAttributes.Num())
+		{
+			// Append the defaults
+			EventAttributes.Append(DefaultEventAttributes);
+		}
+
+		TSharedPtr<FArchive> Archive;
+
+		// If we already have a registered output file for this event then use it, otherwise create a new one and add it to the archives list 
+		// Ensure we have an upper limit of the maximum number of events to export
+		if ( Archives.Find(EventName) == nullptr ) 
+		{
+			if (Archives.Num() >= MaxEvents)
+			{
+				// We can't create any new CSV files as we have exceeded the maximum number of allowed event files.
+				return;
+			}
+			else
+			{
+				// Create the full output path
+				FString FilePath = FolderPath / FPaths::MakeValidFileName(EventName) + TEXT(".csv");
+
+				// Make a new CSV file for this event and add it to the list of files.
+				Archive = TSharedPtr<FArchive>(IFileManager::Get().CreateFileWriter(*FilePath, FILEWRITE_EvenIfReadOnly));
+
+				if (Archive.IsValid())
+				{
+					// Add the new;y created archive to the list
+					Archives.Emplace(EventName, Archive);
+
+					// Add the Attribute Names to the first row of the new CSV file
+					TStringBuilder<1024> StringBuilder;
+
+					int32 Column = 0;
+					for (const FAnalyticsEventAttribute& Attribute : EventAttributes)
+					{
+						if (Column == 0)
+						{
+							StringBuilder.Appendf(TEXT("%s"), *Attribute.GetName());
+						}
+						else
+						{
+							StringBuilder.Appendf(TEXT(",%s"), *Attribute.GetName());
+						}
+
+						Column++;
+					}
+
+					FString OutputString = StringBuilder.ToString();
+
+					WriterPipe->Launch(TEXT("FAnalyticsProviderCSV_WriteJob"), [Archive, OutputString = MoveTemp(OutputString)]()
+						{
+							Archive->Logf(TEXT("%s"), *OutputString);
+							Archive->Flush();
+						});
+				}
+			}
+		}
+		else
+		{
+			// Retrieve the existing Archive
+			Archive = Archives[EventName];
+		}
+		
+		if (Archive.IsValid())
+		{
+			// Append the Attribute Values to the subsequent rows of the existing CSV file
+			TStringBuilder<1024> StringBuilder;
+
+			int32 Column = 0;
+			for (const FAnalyticsEventAttribute& Attribute : EventAttributes)
+			{
+				if (Column == 0)
+				{
+					StringBuilder.Appendf(TEXT("%s"), *Attribute.GetValue());
+				}
+				else
+				{
+					StringBuilder.Appendf(TEXT(",%s"), *Attribute.GetValue());
+				}
+
+				Column++;
+			}
+
+			FString OutputString = StringBuilder.ToString();
+
+			WriterPipe->Launch(TEXT("FAnalyticsProviderCSV_WriteJob"), [Archive, OutputString = MoveTemp(OutputString)]()
+				{
+					Archive->Logf(TEXT("%s"), *OutputString);
+					Archive->Flush();
+				});
+
+		}
+	}
+}
diff --git a/Engine/Source/Runtime/Analytics/AnalyticsCSV/Private/AnalyticsProviderCSV.h b/Engine/Source/Runtime/Analytics/AnalyticsCSV/Private/AnalyticsProviderCSV.h
new file mode 100644
--- /dev/null
+++ b/Engine/Source/Runtime/Analytics/AnalyticsCSV/Private/AnalyticsProviderCSV.h
@@ -0,0 +1,53 @@
+// Copyright Epic Games, Inc. All Rights Reserved.
+
+#pragma once
+
+#include "AnalyticsProviderConfigurationDelegate.h"
+#include "Interfaces/IAnalyticsProvider.h"
+
+namespace UE::Tasks { class FPipe; }
+
+/**
+ * Implementation of the IAnalyticsProvider interface that exports telemetry events to named CSV files
+ * By default, the log files are written to Saved/Telemetry folder of the application.
+ * FolderPath can be overridden in the configuration
+ * For safety, there is MaxEvents number of event files that will be opened 
+ */
+class FAnalyticsProviderCSV : public IAnalyticsProvider
+{
+public:
+
+	FAnalyticsProviderCSV(const FAnalyticsProviderConfigurationDelegate& GetConfigValue);
+	~FAnalyticsProviderCSV();
+
+	virtual bool StartSession(const TArray<FAnalyticsEventAttribute>& Attributes = {}) override;
+	virtual void EndSession() override;
+
+	virtual void FlushEvents() override;
+
+	virtual void SetDefaultEventAttributes(TArray<FAnalyticsEventAttribute>&& Attributes)  override;
+	virtual TArray<FAnalyticsEventAttribute> GetDefaultEventAttributesSafe() const override;
+	virtual int32 GetDefaultEventAttributeCount() const  override;
+	virtual FAnalyticsEventAttribute GetDefaultEventAttribute(int AttributeIndex) const  override;
+	virtual bool SetSessionID(const FString& InSessionID) override;
+	virtual void SetUserID(const FString& InUserID) override;
+
+	virtual FString GetSessionID() const override;
+	virtual FString GetUserID() const override;
+	
+	virtual void RecordEvent(const FString& EventName, const TArray<FAnalyticsEventAttribute>& Attributes = {}) override;
+
+private:
+
+	using TArchives = TMap < FString, TSharedPtr<FArchive>>;
+
+	FString									UserID;
+	FString									SessionID;
+	FString									FolderPath;
+	TArray<FAnalyticsEventAttribute>		DefaultEventAttributes;	
+	int32									MaxEvents = 16;
+	TUniquePtr<UE::Tasks::FPipe>			WriterPipe;
+	TArchives								Archives;
+	bool									SessionStarted = false;
+	
+};
diff --git a/Engine/Source/Runtime/Analytics/AnalyticsCSV/Public/AnalyticsCSV.h b/Engine/Source/Runtime/Analytics/AnalyticsCSV/Public/AnalyticsCSV.h
new file mode 100644
--- /dev/null
+++ b/Engine/Source/Runtime/Analytics/AnalyticsCSV/Public/AnalyticsCSV.h
@@ -0,0 +1,46 @@
+// Copyright Epic Games, Inc. All Rights Reserved.
+
+#pragma once
+
+#include "AnalyticsProviderConfigurationDelegate.h"
+#include "Containers/Array.h"
+#include "Containers/UnrealString.h"
+#include "CoreMinimal.h"
+#include "HAL/Platform.h"
+#include "HAL/PlatformCrt.h"
+#include "Interfaces/IAnalyticsProviderModule.h"
+#include "Modules/ModuleManager.h"
+#include "Templates/SharedPointer.h"
+#include "Templates/UnrealTemplate.h"
+
+class IAnalyticsProvider;
+
+/**
+ *  Public implementation of IAnalyticsProviderModule that returns a FAnalyticsProviderCSV object
+ */
+class FAnalyticsCSV : public IAnalyticsProviderModule
+{
+	//--------------------------------------------------------------------------
+	// Module functionality
+	//--------------------------------------------------------------------------
+public:
+	/**
+	 * Singleton-like access to this module's interface.  This is just for convenience!
+	 * Beware of calling this during the shutdown phase, though.  Your module might have been unloaded already
+	 * @return Returns singleton instance, loading the module on demand if needed
+	 */
+	static inline FAnalyticsCSV& Get()
+	{
+		return FModuleManager::LoadModuleChecked< FAnalyticsCSV >( "AnalyticsCSV" );
+	}
+
+	/**
+	 * IAnalyticsProviderModule interface.
+	 */
+	ANALYTICSCSV_API virtual TSharedPtr<IAnalyticsProvider> CreateAnalyticsProvider(const FAnalyticsProviderConfigurationDelegate& GetConfigValue) const override;
+
+private:
+	ANALYTICSCSV_API virtual void StartupModule() override;
+	ANALYTICSCSV_API virtual void ShutdownModule() override;
+};
+
diff --git a/Engine/Source/Runtime/Analytics/AnalyticsLog/Private/AnalyticsProviderLog.cpp b/Engine/Source/Runtime/Analytics/AnalyticsLog/Private/AnalyticsProviderLog.cpp
--- a/Engine/Source/Runtime/Analytics/AnalyticsLog/Private/AnalyticsProviderLog.cpp
+++ b/Engine/Source/Runtime/Analytics/AnalyticsLog/Private/AnalyticsProviderLog.cpp
@@ -32,18 +32,29 @@ FAnalyticsProviderLog::FAnalyticsProviderLog(const FAnalyticsProviderConfigurati
 	}
 
 	// Create the full output path
-	FString FilePath = FolderPath / FileName;
-	FileWriter = TUniquePtr<FArchive>(IFileManager::Get().CreateFileWriter(*FilePath, FILEWRITE_EvenIfReadOnly));
-
-	if (FileWriter.IsValid())
+	if (FolderPath.IsEmpty()==false)
 	{
-		WriterPipe = MakeUnique<UE::Tasks::FPipe>(TEXT("FAnalyticsProviderLog_Writer"));
+		if (FolderPath.IsEmpty() == false)
+		{
+			if (IFileManager::Get().MakeDirectory(*FolderPath, true))
+			{
+				WriterPipe = MakeUnique<UE::Tasks::FPipe>(TEXT("FAnalyticsProviderLog_Writer"));
+
+				FString FilePath = FolderPath / FileName;
+				FileWriter = TUniquePtr<FArchive>(IFileManager::Get().CreateFileWriter(*FilePath, FILEWRITE_EvenIfReadOnly));
+			}
+		}
 	}
 }
 
 FAnalyticsProviderLog::~FAnalyticsProviderLog()
 {
-
+	if (WriterPipe.IsValid())
+	{
+		TRACE_CPUPROFILER_EVENT_SCOPE(FAnalyticsProviderLog::EndSession::ClosePipe);
+		WriterPipe->WaitUntilEmpty();
+		WriterPipe.Reset();
+	}
 }
 
 bool FAnalyticsProviderLog::SetSessionID(const FString& InSessionID)
@@ -89,7 +100,7 @@ int32 FAnalyticsProviderLog::GetDefaultEventAttributeCount() const
 
 FAnalyticsEventAttribute FAnalyticsProviderLog::GetDefaultEventAttribute(int AttributeIndex) const
 {
-	return DefaultEventAttributes[AttributeIndex];
+	return AttributeIndex < DefaultEventAttributes.Num() ? DefaultEventAttributes[AttributeIndex] : FAnalyticsEventAttribute();
 }
 
 bool FAnalyticsProviderLog::StartSession(const TArray<FAnalyticsEventAttribute>& Attributes)
@@ -109,9 +120,7 @@ void FAnalyticsProviderLog::EndSession()
 
 	if (WriterPipe.IsValid())
 	{
-		TRACE_CPUPROFILER_EVENT_SCOPE(FAnalyticsProviderLog::EndSession::ClosePipe);
 		WriterPipe->WaitUntilEmpty();
-		WriterPipe.Reset();
 	}
 
 	if (FileWriter.IsValid())
@@ -143,8 +152,13 @@ void FAnalyticsProviderLog::RecordEvent(const FString& EventName, const TArray<F
 		Builder.Appendf(TEXT(",\"RecordId\":%d"), RecordId++);
 
 		// Accumulate all the attributes together. We could have had two loops but this seems cleaner
-		TArray<FAnalyticsEventAttribute> EventAttributes(DefaultEventAttributes);
-		EventAttributes.Append(Attributes);
+		TArray<FAnalyticsEventAttribute> EventAttributes(Attributes);
+
+		if (DefaultEventAttributes.Num())
+		{
+			// Append the defaults
+			EventAttributes.Append(DefaultEventAttributes);
+		}
 
 		// Add all the attributes
 		for (const FAnalyticsEventAttribute& Attribute : EventAttributes)
diff --git a/Engine/Source/Runtime/Analytics/AnalyticsLog/Private/AnalyticsProviderLog.h b/Engine/Source/Runtime/Analytics/AnalyticsLog/Private/AnalyticsProviderLog.h
--- a/Engine/Source/Runtime/Analytics/AnalyticsLog/Private/AnalyticsProviderLog.h
+++ b/Engine/Source/Runtime/Analytics/AnalyticsLog/Private/AnalyticsProviderLog.h
@@ -8,7 +8,7 @@
 namespace UE::Tasks { class FPipe; }
 
 /**
- * Implementation of the IAnalyticsProviderET interface that exports telemetry events to a file as Newline - delimited JSON
+ * Implementation of the IAnalyticsProvider interface that exports telemetry events to a file as Newline - delimited JSON
  * By default,the log file is written to Saved/Telemetry folder of the application.
  * FileName and FolderPath can be overridden in the configuration
  * Here as a simple example of how a developer might implement and configure their own analytics provider for use with the StudioTelemerty plugin
diff --git a/Engine/Source/Runtime/StudioTelemetry/StudioTelemetry.build.cs b/Engine/Source/Runtime/StudioTelemetry/StudioTelemetry.build.cs
--- a/Engine/Source/Runtime/StudioTelemetry/StudioTelemetry.build.cs
+++ b/Engine/Source/Runtime/StudioTelemetry/StudioTelemetry.build.cs
@@ -15,9 +15,10 @@ public class StudioTelemetry : ModuleRules
 				"HTTP",
                 "BuildSettings",
                 "Analytics",
-                "AnalyticsET",
-				"AnalyticsLog",
+				"AnalyticsCSV",
+				"AnalyticsET",
 				"AnalyticsHorde",
+				"AnalyticsLog",
 				"TelemetryUtils"
 			}
 		);
