Fix for PVS warning

#Synced-CL 48759921

diff --git a/Engine/Source/Runtime/Analytics/AnalyticsLog/Private/AnalyticsProviderLog.cpp b/Engine/Source/Runtime/Analytics/AnalyticsLog/Private/AnalyticsProviderLog.cpp
--- a/Engine/Source/Runtime/Analytics/AnalyticsLog/Private/AnalyticsProviderLog.cpp
+++ b/Engine/Source/Runtime/Analytics/AnalyticsLog/Private/AnalyticsProviderLog.cpp
@@ -32,17 +32,14 @@ FAnalyticsProviderLog::FAnalyticsProviderLog(const FAnalyticsProviderConfigurati
 	}
 
 	// Create the full output path
-	if (FolderPath.IsEmpty()==false)
+	if (FolderPath.IsEmpty() == false)
 	{
-		if (FolderPath.IsEmpty() == false)
+		if (IFileManager::Get().MakeDirectory(*FolderPath, true))
 		{
-			if (IFileManager::Get().MakeDirectory(*FolderPath, true))
-			{
-				WriterPipe = MakeUnique<UE::Tasks::FPipe>(TEXT("FAnalyticsProviderLog_Writer"));
+			WriterPipe = MakeUnique<UE::Tasks::FPipe>(TEXT("FAnalyticsProviderLog_Writer"));
 
-				FString FilePath = FolderPath / FileName;
-				FileWriter = TUniquePtr<FArchive>(IFileManager::Get().CreateFileWriter(*FilePath, FILEWRITE_EvenIfReadOnly));
-			}
+			FString FilePath = FolderPath / FileName;
+			FileWriter = TUniquePtr<FArchive>(IFileManager::Get().CreateFileWriter(*FilePath, FILEWRITE_EvenIfReadOnly));
 		}
 	}
 }
